**New**
- Engine
	- IAnimationDataModel interface
- AnimationData plugin
	- AnimSequencerDataModel, sequencer based implementation of IAnimationDataModel
	- AnimSequencerDataController, controller for above (implementation of IAnimationDataController)
- Added FCompressedRichCurve::PopulateCurve, allowing users to convert back to a FRichCurve for validation/debugging
- Added DefaultFrameRate to AnimationSettings, this replaces the hardcoded 30FPS in code
- Added ::GetKeyIndicesFromTime which takes FFrameRate
- Added AnimSequenceBase::OnAnimModelLoaded, this is required for correct postloading behaviour of data model (specifically AnimSequencerDataModel, as it depends on data from its outer ? anim sequence)
- Added IAnimationDataModel ::Evaluate which now takes responsibility for evaluating raw animation bone, curve and attribute data. And moved all trackbased evaluation code into AnimDataModel.cpp
- Added a.ForceEvalRawData allowing to force evaluation of source data
- Added a.SkipDDC to force compressing animation data locally

**Changed**
- Reparent UAnimDataModel to IAnimationDataModel interface, and rejig behaviour
- AnimSequenceBase now reference AnimDataModel as IAnimationDataModel
- Upgrade path for legacy to IAnimationDataModel and existing UAnimDataModel to IAnimationDataModel through IAnimationDataController::PopulateWithExisting
- IAnimationDataModel data is now frame(number/rate/time) based rather than seconds/keyindices. This enforces frameborder aligned data from now on.
- Moved RichCurve evaluation code from RichCurve.cpp to separate file and consolidated all instances of copy-pasta behaviour
- Updated/removed deprecated EngineTests around AnimSequences
- Changed many instances to use a FFrameTime together with a known FFrameRate versus seconds and sequencelength in float involving frame \/ time calculations (including instances of FAnimKeyHelper)
- Switched anim sequence evaluation to use double in Extraction context for time value and patched up places with static\_cast\<double\>
- FRichCurve::Eval
	- Make sure we always evaluate keys when T >= Key0.Time and T <= KeyN.Time to deal with crazy userweighted tangents
- Fixed WeightedKeyDataAdapter::GetKeyInterpMode/GetKeyTangentWeightMode retrieving invalid values, as it was indexing according to KeyIndex rather than KeyIndex\*2. This made it so that values were misinterpreted.
- Fixed WeightedEvalForTwoKeys for compressed data retrieving GetKeyInterpMode and GetKeyTangentWeightMode using the wrong index value. As they are not indexed using KeyDataHandle type.
- Fixed issue where compression could crash when containing zero-key scale additivebase track (tries to retrieve)
- Replaced instances of NULL with nullptr
- Moved required decompression information into FAnimSequenceDecompressionContext and reimplemented decompression within UE::Anim::Decompression namespace and new file
- Deprecated ::GetRawDataGuid and replaced with GetDataModel()->GenerateGuid()
- Fixed UAnimStreamable evaluation/compression (previously broken with MVC refactor)
- Updated Animation exporting pipeline to be frame-based rather than seconds
- Deprecated RetargetPose, replaced with FRetargetingScope (used within AnimModel::Evaluate)
- Fixed BaseAdditiveAnimation array become invalid/incorrect when removing zero-additve tracks during compression
- Updated EngineTest animation sequence related tests to new APIs (while maintaining deprecated path testing as well for now)
- AnimDataController / Animation Sequence tests now generate a USkeleton for the transient animation sequence (used to perform the test on)

**Removed**
- Unused file/class AnimData/AnimDataNotifyCollector.h

diff --git a/Engine/Plugins/Animation/AnimationData/Source/AnimationData/Private/AnimSequencerDataModel.cpp b/Engine/Plugins/Animation/AnimationData/Source/AnimationData/Private/AnimSequencerDataModel.cpp
new file mode 100644
--- /dev/null
+++ b/Engine/Plugins/Animation/AnimationData/Source/AnimationData/Private/AnimSequencerDataModel.cpp
@@ -0,0 +1,1494 @@
+// Copyright Epic Games, Inc. All Rights Reserved.
+
+#include "AnimSequencerDataModel.h"
+#include "AnimSequencerController.h"
+
+#include "Animation/AnimSequence.h"
+#include "AnimDataController.h"
+#include "IAnimationEditor.h"
+#include "ControlRig.h"
+#include "ControlRigObjectBinding.h"
+#include "Algo/Accumulate.h"
+#include "Animation/AnimSequenceHelpers.h"
+#include "Rigs/FKControlRig.h"
+#include "Sequencer/MovieSceneControlRigParameterSection.h"
+#include "Sequencer/MovieSceneControlRigParameterTrack.h"
+#include "UObject/ObjectSaveContext.h"
+#include "Animation/AnimData/AnimDataNotifications.h"
+
+#include "AnimSequencerHelpers.h"
+#include "Animation/AnimationSettings.h"
+#include "MovieScene/Private/Channels/MovieSceneCurveChannelImpl.h"
+
+#define LOCTEXT_NAMESPACE "AnimSequencerDataModel"
+
+int32 UAnimationSequencerDataModel::RetainFloatCurves = 0;
+static FAutoConsoleVariableRef CVarRetainFloatCurves(
+	TEXT("a.AnimSequencer.RetainFloatCurves"),
+	UAnimationSequencerDataModel::RetainFloatCurves,
+	TEXT("1 = Original FloatCurves are retained when generating transient curve data from Control Curves . 0 = FloatCurves are overriden with Control Curves"));
+
+int32 UAnimationSequencerDataModel::ValidationMode = 0;
+static FAutoConsoleVariableRef CValidationMode(
+	TEXT("a.AnimSequencer.ValidationMode"),
+	UAnimationSequencerDataModel::ValidationMode,
+	TEXT("1 = Enables validation after operations to test data integrity against legacy version. 0 = validation disabled"));
+
+int32 UAnimationSequencerDataModel::UseDirectFKControlRigMode = 1;
+static FAutoConsoleVariableRef CVarDirectControlRigMode(
+	TEXT("a.AnimSequencer.DirectControlRigMode"),
+	UAnimationSequencerDataModel::UseDirectFKControlRigMode,
+	TEXT("1 = FKControl rig uses Direct method for setting Control transforms. 0 = FKControl rig uses Replace method (transform offsets) for setting Control transforms"));
+
+void UAnimationSequencerDataModel::RemoveOutOfDateControls() const
+{
+	if (UMovieSceneControlRigParameterSection* Section = GetFKControlRigSection())
+	{
+		if (UFKControlRig* ControlRig = Cast<UFKControlRig>(Section->GetControlRig()))
+		{
+			if (URigHierarchy* Hierarchy = ControlRig->GetHierarchy())
+			{
+				if (URigHierarchyController* Controller = Hierarchy->GetController())
+				{
+					TArray<FRigElementKey> ElementKeysToRemove;
+					Hierarchy->ForEach<FRigControlElement>([this, Section, &ElementKeysToRemove](const FRigControlElement* ControlElement) -> bool
+					{
+						const bool bContainsBone = Section->HasTransformParameter(ControlElement->GetName());
+						const bool bContainsCurve = Section->HasScalarParameter(ControlElement->GetName());
+						
+						if (!bContainsBone && !bContainsCurve)
+						{
+							ElementKeysToRemove.Add(ControlElement->GetKey());
+						}
+						
+						return true;
+					});
+						
+					Hierarchy->ForEach<FRigCurveElement>([this, &ElementKeysToRemove](const FRigCurveElement* CurveElement) -> bool
+					{
+						const FName TargetCurveName = CurveElement->GetName();
+						if(!LegacyCurveData.FloatCurves.ContainsByPredicate([TargetCurveName](const FFloatCurve& Curve) { return Curve.Name.DisplayName == TargetCurveName; }))
+						{
+							ElementKeysToRemove.Add(CurveElement->GetKey());	
+						}
+						return true;
+					});
+
+					for (const FRigElementKey& KeyToRemove : ElementKeysToRemove)
+					{
+						Controller->RemoveElement(KeyToRemove);
+					}
+
+					ControlRig->RefreshActiveControls();
+				}
+			}
+		}
+	}
+}
+
+USkeleton* UAnimationSequencerDataModel::GetSkeleton() const
+{
+	UAnimationAsset* AnimationAsset = CastChecked<UAnimationAsset>(GetOuter());	
+	checkf(AnimationAsset, TEXT("Unable to retrieve owning AnimationAsset"));
+
+	USkeleton* Skeleton = AnimationAsset->GetSkeleton();
+	if (Skeleton == nullptr)
+	{
+		IAnimationDataController::ReportObjectErrorf(this, LOCTEXT("UnableToFindSkeleton", "Unable to retrieve target USkeleton for Animation Asset ({0})"), FText::FromString(*AnimationAsset->GetPathName()));
+	} 
+
+	return Skeleton;
+}
+
+void UAnimationSequencerDataModel::InitializeFKControlRig(UFKControlRig* FKControlRig, USkeleton* Skeleton) const
+{
+	checkf(FKControlRig, TEXT("Invalid FKControlRig provided"));
+	if (Skeleton)
+	{
+		FKControlRig->SetObjectBinding(MakeShared<FControlRigObjectBinding>());
+		FKControlRig->GetObjectBinding()->BindToObject(Skeleton);
+
+		UFKControlRig::FRigElementInitializationOptions InitOptions;
+		InitOptions.bImportCurves = false;	
+		if(UMovieSceneControlRigParameterSection* Section = GetFKControlRigSection())
+		{
+			for (const FScalarParameterNameAndCurve& AnimCurve : Section->GetScalarParameterNamesAndCurves())
+			{
+				InitOptions.CurveNames.Add(UFKControlRig::GetControlTargetName(AnimCurve.ParameterName, ERigElementType::Curve));
+			}
+
+			for (const FTransformParameterNameAndCurves& BoneCurve : Section->GetTransformParameterNamesAndCurves())
+			{
+				InitOptions.BoneNames.Add(UFKControlRig::GetControlTargetName(BoneCurve.ParameterName, ERigElementType::Bone));
+			}
+		}
+		InitOptions.bGenerateBoneControls = InitOptions.BoneNames.Num() > 0;
+		FKControlRig->SetInitializationOptions(InitOptions);
+
+		FKControlRig->Initialize();
+
+		FKControlRig->SetApplyMode(UseDirectFKControlRigMode == 1 ? EControlRigFKRigExecuteMode::Direct : EControlRigFKRigExecuteMode::Replace);
+		FKControlRig->SetBoneInitialTransformsFromRefSkeleton(Skeleton->GetReferenceSkeleton());
+		FKControlRig->Evaluate_AnyThread();
+	}
+}
+
+UControlRig* UAnimationSequencerDataModel::GetControlRig() const
+{
+	if(const UMovieSceneControlRigParameterTrack* Track = GetControlRigTrack())
+	{
+		return Track->GetControlRig();
+	}
+
+	return nullptr;
+}
+
+void UAnimationSequencerDataModel::PostLoad()
+{
+	Super::PostLoad();
+
+	if (!HasAnyFlags(RF_ClassDefaultObject))
+	{
+		// Pre/post load any dependencies (Sequencer objects)
+		TArray<UObject*> ObjectReferences;
+		FReferenceFinder(ObjectReferences, this, false, true, true, true).FindReferences(this);	
+		for (UObject* Dependency : ObjectReferences)
+		{
+			if (Dependency->HasAnyFlags(RF_NeedLoad))
+			{
+				Dependency->GetLinker()->Preload(Dependency);
+			}
+
+			if (Dependency->HasAnyFlags(RF_NeedPostLoad))
+			{
+				Dependency->ConditionalPostLoad();
+			}
+		}
+
+		if (const UMovieSceneControlRigParameterSection* Section = GetFKControlRigSection())
+		{
+			if (UFKControlRig* ControlRig = Cast<UFKControlRig>(Section->GetControlRig()))
+			{
+				InitializeFKControlRig(ControlRig, GetSkeleton());
+			}
+		}
+
+		RemoveOutOfDateControls();
+
+		ValidateData();
+	}
+}
+
+void UAnimationSequencerDataModel::GetPreloadDependencies(TArray<UObject*>& OutDeps)
+{
+	Super::GetPreloadDependencies(OutDeps);
+	OutDeps.Add(MovieScene);
+}
+
+void UAnimationSequencerDataModel::PostDuplicate(bool bDuplicateForPIE)
+{
+	UObject::PostDuplicate(bDuplicateForPIE);
+
+	GetNotifier().Notify(EAnimDataModelNotifyType::Populated);
+}
+
+double UAnimationSequencerDataModel::GetPlayLength() const
+{
+	ValidateSequencerData();
+	return MovieScene->GetDisplayRate().AsSeconds(GetNumberOfFrames());
+}
+
+int32 UAnimationSequencerDataModel::GetNumberOfFrames() const
+{	
+	ValidateSequencerData();
+	const TRange<FFrameNumber> FrameRange = MovieScene->GetPlaybackRange();	
+	const TRangeBound<FFrameNumber>& UpperRange = FrameRange.GetUpperBound();
+	const bool bInclusive = UpperRange.IsInclusive();
+	int32 Value = UpperRange.GetValue().Value;
+	if (!bInclusive)
+	{		
+		Value = FMath::Max(Value - 1, 1);
+	}
+
+	return Value;
+}
+
+int32 UAnimationSequencerDataModel::GetNumberOfKeys() const
+{
+	return GetNumberOfFrames() + 1;
+}
+
+FFrameRate UAnimationSequencerDataModel::GetFrameRate() const
+{	
+	ValidateSequencerData();
+	return MovieScene->GetDisplayRate();
+}
+
+const TArray<FBoneAnimationTrack>& UAnimationSequencerDataModel::GetBoneAnimationTracks() const
+{
+	return LegacyBoneAnimationTracks;
+}
+
+const FBoneAnimationTrack& UAnimationSequencerDataModel::GetBoneTrackByIndex(int32 TrackIndex) const
+{
+	checkf(LegacyBoneAnimationTracks.IsValidIndex(TrackIndex), TEXT("Unable to find animation track by index"));
+	return LegacyBoneAnimationTracks[TrackIndex];
+}
+
+const FBoneAnimationTrack& UAnimationSequencerDataModel::GetBoneTrackByName(FName TrackName) const
+{
+	const FBoneAnimationTrack* TrackPtr = LegacyBoneAnimationTracks.FindByPredicate([TrackName](const FBoneAnimationTrack& Track)
+	{
+		return Track.Name == TrackName;
+	});
+
+	checkf(TrackPtr != nullptr, TEXT("Unable to find animation track by name"));
+
+	return *TrackPtr;
+}
+
+const FBoneAnimationTrack* UAnimationSequencerDataModel::FindBoneTrackByName(FName Name) const
+{
+	return LegacyBoneAnimationTracks.FindByPredicate([Name](const FBoneAnimationTrack& Track)
+	{
+		return Track.Name == Name;
+	});
+}
+
+
+FBoneAnimationTrack* UAnimationSequencerDataModel::FindMutableBoneTrackByName(FName Name)
+{
+	return LegacyBoneAnimationTracks.FindByPredicate([Name](const FBoneAnimationTrack& Track)
+	{
+		return Track.Name == Name;
+	});
+}
+
+const FBoneAnimationTrack* UAnimationSequencerDataModel::FindBoneTrackByIndex(int32 BoneIndex) const
+{
+	const FBoneAnimationTrack* TrackPtr = LegacyBoneAnimationTracks.FindByPredicate([BoneIndex](const FBoneAnimationTrack& Track)
+	{
+		return Track.BoneTreeIndex == BoneIndex;
+	});
+
+	return TrackPtr;
+}
+
+int32 UAnimationSequencerDataModel::GetBoneTrackIndex(const FBoneAnimationTrack& Track) const
+{
+	return LegacyBoneAnimationTracks.IndexOfByPredicate([&Track](const FBoneAnimationTrack SearchTrack)
+	{
+		return SearchTrack.Name == Track.Name;
+	});
+}
+
+int32 UAnimationSequencerDataModel::GetBoneTrackIndexByName(FName TrackName) const
+{
+	if (const FBoneAnimationTrack* TrackPtr = FindBoneTrackByName(TrackName))
+	{
+		return GetBoneTrackIndex(*TrackPtr);
+	}
+
+	return INDEX_NONE;
+}
+
+bool UAnimationSequencerDataModel::IsValidBoneTrackIndex(int32 TrackIndex) const
+{
+	return LegacyBoneAnimationTracks.IsValidIndex(TrackIndex);
+}
+
+int32 UAnimationSequencerDataModel::GetNumBoneTracks() const
+{
+	return LegacyBoneAnimationTracks.Num();
+}
+
+void UAnimationSequencerDataModel::GetBoneTrackNames(TArray<FName>& OutNames) const
+{
+	Algo::Transform(LegacyBoneAnimationTracks, OutNames, [](FBoneAnimationTrack Track)
+	{
+		return Track.Name; 
+	});
+}
+
+const FAnimationCurveData& UAnimationSequencerDataModel::GetCurveData() const
+{
+	return LegacyCurveData;
+}
+
+int32 UAnimationSequencerDataModel::GetNumberOfTransformCurves() const
+{
+	return LegacyCurveData.TransformCurves.Num();
+}
+
+int32 UAnimationSequencerDataModel::GetNumberOfFloatCurves() const
+{
+	return LegacyCurveData.FloatCurves.Num();
+}
+
+const TArray<FFloatCurve>& UAnimationSequencerDataModel::GetFloatCurves() const
+{
+	return LegacyCurveData.FloatCurves;
+}
+
+const TArray<FTransformCurve>& UAnimationSequencerDataModel::GetTransformCurves() const
+{
+	return LegacyCurveData.TransformCurves;
+}
+
+const FAnimCurveBase* UAnimationSequencerDataModel::FindCurve(const FAnimationCurveIdentifier& CurveIdentifier) const
+{
+	switch (CurveIdentifier.CurveType)
+	{
+	case ERawCurveTrackTypes::RCT_Float:
+		return FindFloatCurve(CurveIdentifier);
+	case ERawCurveTrackTypes::RCT_Transform:
+		return FindTransformCurve(CurveIdentifier);
+	default:
+		checkf(false, TEXT("Invalid curve identifier type"));
+	}
+
+	return nullptr;
+}
+
+const FFloatCurve* UAnimationSequencerDataModel::FindFloatCurve(const FAnimationCurveIdentifier& CurveIdentifier) const
+{
+	ensure(CurveIdentifier.CurveType == ERawCurveTrackTypes::RCT_Float);
+	for (const FFloatCurve& FloatCurve : GetCurveData().FloatCurves)
+	{
+		if (FloatCurve.Name == CurveIdentifier.InternalName || (FloatCurve.Name.UID == CurveIdentifier.InternalName.UID && FloatCurve.Name.UID != SmartName::MaxUID))
+		{
+			return &FloatCurve;
+		}
+	}
+
+	return nullptr;
+}
+
+const FTransformCurve* UAnimationSequencerDataModel::FindTransformCurve(const FAnimationCurveIdentifier& CurveIdentifier) const
+{
+	ensure(CurveIdentifier.CurveType == ERawCurveTrackTypes::RCT_Transform);
+    for (const FTransformCurve& TransformCurve : GetCurveData().TransformCurves)
+    {
+    	if (TransformCurve.Name == CurveIdentifier.InternalName || TransformCurve.Name.UID == CurveIdentifier.InternalName.UID)
+    	{
+    		return &TransformCurve;
+    	}
+    }
+
+    return nullptr;
+}
+
+const FRichCurve* UAnimationSequencerDataModel::FindRichCurve(const FAnimationCurveIdentifier& CurveIdentifier) const
+{
+	const FRichCurve* RichCurve = nullptr;
+
+	if (CurveIdentifier.IsValid())
+	{
+		if (CurveIdentifier.CurveType == ERawCurveTrackTypes::RCT_Float)
+		{
+			const FFloatCurve* Curve = FindFloatCurve(CurveIdentifier);
+			if (Curve)
+			{
+				RichCurve = &Curve->FloatCurve;
+			}
+		}
+		else if (CurveIdentifier.CurveType == ERawCurveTrackTypes::RCT_Transform)
+		{
+			if (CurveIdentifier.Channel != ETransformCurveChannel::Invalid && CurveIdentifier.Axis != EVectorCurveChannel::Invalid)
+			{
+				// Dealing with transform curve
+				if (const FTransformCurve* TransformCurve = FindTransformCurve(CurveIdentifier))
+				{
+					if (const FVectorCurve* VectorCurve = TransformCurve->GetVectorCurveByIndex(static_cast<int32>(CurveIdentifier.Channel)))
+					{
+						RichCurve = &VectorCurve->FloatCurves[static_cast<int32>(CurveIdentifier.Axis)];
+					}
+				}
+
+			}
+		}
+	}
+
+	return RichCurve;
+}
+
+const FAnimCurveBase& UAnimationSequencerDataModel::GetCurve(const FAnimationCurveIdentifier& CurveIdentifier) const
+{
+	const FAnimCurveBase* CurvePtr = FindCurve(CurveIdentifier);
+	checkf(CurvePtr, TEXT("Tried to retrieve non-existing curve"));
+	return *CurvePtr;
+}
+
+const FFloatCurve& UAnimationSequencerDataModel::GetFloatCurve(const FAnimationCurveIdentifier& CurveIdentifier) const
+{
+	const FFloatCurve* CurvePtr = FindFloatCurve(CurveIdentifier);
+	checkf(CurvePtr, TEXT("Tried to retrieve non-existing curve"));
+	return *CurvePtr;
+}
+
+const FTransformCurve& UAnimationSequencerDataModel::GetTransformCurve(const FAnimationCurveIdentifier& CurveIdentifier) const
+{
+	const FTransformCurve* CurvePtr = FindTransformCurve(CurveIdentifier);
+	checkf(CurvePtr, TEXT("Tried to retrieve non-existing curve"));
+	return *CurvePtr;
+}
+
+const FRichCurve& UAnimationSequencerDataModel::GetRichCurve(const FAnimationCurveIdentifier& CurveIdentifier) const
+{
+	const FRichCurve* CurvePtr = FindRichCurve(CurveIdentifier);
+	checkf(CurvePtr, TEXT("Tried to retrieve non-existing curve"));
+	return *CurvePtr;	
+}
+
+TArrayView<const FAnimatedBoneAttribute> UAnimationSequencerDataModel::GetAttributes() const
+{
+	return AnimatedBoneAttributes;
+}
+
+int32 UAnimationSequencerDataModel::GetNumberOfAttributes() const
+{
+	return AnimatedBoneAttributes.Num();
+}
+
+int32 UAnimationSequencerDataModel::GetNumberOfAttributesForBoneIndex(const int32 BoneIndex) const
+{
+	// Sum up total number of attributes with provided bone index
+	const int32 NumberOfBoneAttributes = Algo::Accumulate<int32>(AnimatedBoneAttributes, 0, [BoneIndex](int32 Sum, const FAnimatedBoneAttribute& Attribute) -> int32
+	{
+		Sum += Attribute.Identifier.GetBoneIndex() == BoneIndex ? 1 : 0;
+		return Sum;
+	});
+	return NumberOfBoneAttributes;
+}
+
+void UAnimationSequencerDataModel::GetAttributesForBone(const FName& BoneName, TArray<const FAnimatedBoneAttribute*>& OutBoneAttributes) const
+{
+	Algo::TransformIf(AnimatedBoneAttributes, OutBoneAttributes, [BoneName](const FAnimatedBoneAttribute& Attribute) -> bool
+	{
+		return Attribute.Identifier.GetBoneName() == BoneName;
+	},
+	[](const FAnimatedBoneAttribute& Attribute) 
+	{
+		return &Attribute;
+	});
+}
+
+const FAnimatedBoneAttribute& UAnimationSequencerDataModel::GetAttribute(const FAnimationAttributeIdentifier& AttributeIdentifier) const
+{
+	const FAnimatedBoneAttribute* AttributePtr = FindAttribute(AttributeIdentifier);
+	checkf(AttributePtr, TEXT("Unable to find attribute for provided identifier"));
+
+	return *AttributePtr;
+}
+
+const FAnimatedBoneAttribute* UAnimationSequencerDataModel::FindAttribute(const FAnimationAttributeIdentifier& AttributeIdentifier) const
+{
+	return AnimatedBoneAttributes.FindByPredicate([AttributeIdentifier](const FAnimatedBoneAttribute& Attribute)
+	{
+		return Attribute.Identifier == AttributeIdentifier;
+	});
+}
+
+UAnimSequence* UAnimationSequencerDataModel::GetAnimationSequence() const
+{
+	return Cast<UAnimSequence>(GetOuter());
+}
+
+FGuid UAnimationSequencerDataModel::GenerateGuid() const
+{	
+	if (CachedRawDataGUID.IsValid())
+	{
+		return CachedRawDataGUID;
+	}
+	else
+	{
+		FSHA1 Sha;
+		const FString ClassName = GetClass()->GetName();
+		Sha.UpdateWithString(*ClassName, ClassName.Len());
+		
+		auto UpdateSHAWithArray = [&Sha](const auto& Array)
+		{
+			if (Array.Num())
+			{
+				Sha.Update(reinterpret_cast<const uint8*>(Array.GetData()), Array.Num() * Array.GetTypeSize());			
+			}
+		};
+
+		auto UpdateWithChannel = [UpdateSHAWithArray](const auto& Channel)
+		{
+			UpdateSHAWithArray(Channel.GetData().GetTimes());
+			UpdateSHAWithArray(Channel.GetData().GetValues());
+		};
+
+		if (const UMovieSceneControlRigParameterSection* RigSection = GetFKControlRigSection())
+		{
+			UpdateWithChannel(RigSection->Weight);
+
+			for (const FTransformParameterNameAndCurves& TransformParameter : RigSection->GetTransformParameterNamesAndCurves())
+			{
+				const FString ParameterString = TransformParameter.ParameterName.ToString();
+				Sha.UpdateWithString(*ParameterString, ParameterString.Len());
+				for (int32 Index = 0; Index < 3; ++Index)
+				{							
+					UpdateWithChannel(TransformParameter.Translation[Index]);
+					UpdateWithChannel(TransformParameter.Rotation[Index]);
+					UpdateWithChannel(TransformParameter.Scale[Index]);
+				}
+			}
+
+			for (const FScalarParameterNameAndCurve& ScalarCurve : RigSection->GetScalarParameterNamesAndCurves())
+			{
+				const FString ParameterString = ScalarCurve.ParameterName.ToString();
+				Sha.UpdateWithString(*ParameterString, ParameterString.Len());
+				UpdateWithChannel(ScalarCurve.ParameterCurve);
+			}
+		}
+		
+		auto UpdateWithData = [&Sha](const auto& Data)
+		{
+			Sha.Update(reinterpret_cast<const uint8*>(&Data), sizeof(Data));
+		};
+		
+		for (const FAnimatedBoneAttribute& Attribute : AnimatedBoneAttributes)
+		{
+			UpdateWithData(Attribute.Identifier);
+			UpdateSHAWithArray(Attribute.Curve.GetConstRefOfKeys());
+		}
+
+		Sha.Final();
+
+		uint32 Hash[5];
+		Sha.GetHash(reinterpret_cast<uint8*>(Hash));
+		const FGuid Guid(Hash[0] ^ Hash[4], Hash[1], Hash[2], Hash[3]);
+		
+		return Guid;
+	}
+}
+
+TScriptInterface<IAnimationDataController> UAnimationSequencerDataModel::GetController()
+{
+	return NewObject<UAnimSequencerController>();
+}
+
+IAnimationDataModel::FModelNotifier& UAnimationSequencerDataModel::GetNotifier()
+{
+	if (!Notifier)
+	{
+		Notifier.Reset(new IAnimationDataModel::FModelNotifier(this));
+	}
+	
+	return *Notifier.Get();
+}
+
+void UAnimationSequencerDataModel::Evaluate(FAnimationPoseData& InOutPoseData, const UE::Anim::DataModel::FEvaluationContext& EvaluationContext) const
+{
+	QUICK_SCOPE_CYCLE_COUNTER(AnimationDataSequence_Evaluate);
+
+	if(!!ValidationMode)
+	{
+		ValidateSequencerData();		
+	}
+
+	if (UMovieSceneControlRigParameterTrack* Track = GetControlRigTrack())
+	{
+		FScopeLock Lock(&EvaluationLock);
+		// Evaluates and applies control curves from track to ControlRig
+		EvaluateTrack(Track, EvaluationContext);
+
+		// Generate/populate the output animation pose data
+		UControlRig* ControlRig = Track->GetControlRig();
+		GeneratePoseData(ControlRig, InOutPoseData, EvaluationContext);
+	}
+}
+
+void UAnimationSequencerDataModel::OnNotify(const EAnimDataModelNotifyType& NotifyType, const FAnimDataModelNotifPayload& Payload)
+{
+	Collector.Handle(NotifyType);
+
+	if (Collector.IsNotWithinBracket() && bPopulated)
+	{
+		// Once the model has been populated and a modification is made - invalidate the cached GUID
+		auto ResetCachedGUID = [this]()
+		{
+			if (CachedRawDataGUID.IsValid() && !Collector.Contains(EAnimDataModelNotifyType::Populated))
+			{
+				CachedRawDataGUID.Invalidate();			
+			}
+		};		
+		
+		const TArray<EAnimDataModelNotifyType> CurveNotifyTypes = {EAnimDataModelNotifyType::CurveAdded, EAnimDataModelNotifyType::CurveChanged, EAnimDataModelNotifyType::CurveRenamed, EAnimDataModelNotifyType::CurveRemoved,
+			EAnimDataModelNotifyType::CurveFlagsChanged, EAnimDataModelNotifyType::CurveScaled, EAnimDataModelNotifyType::CurveColorChanged, EAnimDataModelNotifyType::Populated, EAnimDataModelNotifyType::Reset };
+		if(Collector.Contains(CurveNotifyTypes))
+		{
+			if(!ValidationMode)
+			{
+				GenerateLegacyCurveData();				
+			}
+			ResetCachedGUID();
+		}
+
+		const TArray<EAnimDataModelNotifyType> BonesNotifyTypes = {EAnimDataModelNotifyType::TrackAdded, EAnimDataModelNotifyType::TrackChanged, EAnimDataModelNotifyType::TrackRemoved, EAnimDataModelNotifyType::Populated, EAnimDataModelNotifyType::Reset };
+		if(Collector.Contains(BonesNotifyTypes))
+		{
+			if(!ValidationMode)
+			{
+				GenerateLegacyBoneData();
+			}
+			ResetCachedGUID();
+		}
+
+		if (Collector.Contains(EAnimDataModelNotifyType::Populated))
+		{
+			if(UMovieSceneControlRigParameterSection* Section = GetFKControlRigSection())
+			{
+				Section->ReconstructChannelProxy();
+			}
+		}
+				
+		ValidateData();		
+	}
+}
+
+UMovieSceneControlRigParameterTrack* UAnimationSequencerDataModel::GetControlRigTrack() const
+{
+	return MovieScene->FindMasterTrack<UMovieSceneControlRigParameterTrack>();
+}
+
+UMovieSceneControlRigParameterSection* UAnimationSequencerDataModel::GetFKControlRigSection() const
+{
+	if (MovieScene)
+	{
+		const UMovieSceneControlRigParameterTrack* Track = GetControlRigTrack();
+		for (UMovieSceneSection* TrackSection : Track->GetAllSections())
+		{
+			if (UMovieSceneControlRigParameterSection* Section = Cast<UMovieSceneControlRigParameterSection>(TrackSection))
+			{
+				if (const UControlRig* ControlRig = Section->GetControlRig())
+				{
+					if (ControlRig->IsA<UFKControlRig>())
+					{
+						return Section;
+					}
+				}
+			}
+		}
+	}
+
+	return nullptr;
+}
+
+void UAnimationSequencerDataModel::GenerateLegacyCurveData()
+{
+	QUICK_SCOPE_CYCLE_COUNTER(STAT_GenerateLegacyCurveData);
+	ValidateSequencerData();
+	
+	if (const UMovieSceneControlRigParameterTrack* Track = GetControlRigTrack())
+	{
+		for (const UMovieSceneSection* TrackSection : Track->GetAllSections())
+		{
+			if (const UMovieSceneControlRigParameterSection* Section = Cast<UMovieSceneControlRigParameterSection>(TrackSection))
+			{
+				if (const UControlRig* ControlRig = Section->GetControlRig())
+				{
+					if (USkeleton* Skeleton = GetSkeleton())
+					{
+						if(URigHierarchy* Hierarchy = ControlRig->GetHierarchy())
+						{
+							const FString SequencerSuffix(TEXT("_Sequencer"));
+							const TArray<FScalarParameterNameAndCurve>& ScalarCurves = Section->GetScalarParameterNamesAndCurves();				
+							if (RetainFloatCurves)
+							{
+								LegacyCurveData.FloatCurves.RemoveAll([SequencerSuffix](const FFloatCurve& FloatCurve)
+								{
+									return FloatCurve.Name.DisplayName.ToString().EndsWith(SequencerSuffix);
+								});
+							}
+							else
+							{
+								LegacyCurveData.FloatCurves.Empty();
+							}				
+
+							Hierarchy->ForEach<FRigCurveElement>([Hierarchy,Skeleton, this, ScalarCurves, FrameRate = GetFrameRate(), SequencerSuffix](const FRigCurveElement* CurveElement) -> bool
+							{
+								const FRigElementKey ControlKey(UFKControlRig::GetControlName(CurveElement->GetName(), ERigElementType::Curve), ERigElementType::Control);
+								if (const FRigControlElement* Element = Hierarchy->Find<FRigControlElement>(ControlKey))
+								{
+									FFloatCurve& FloatCurve = LegacyCurveData.FloatCurves.AddDefaulted_GetRef();
+									if (RetainFloatCurves)
+									{
+										FloatCurve.Name.DisplayName = FName(*(CurveElement->GetName().ToString() + TEXT("_Sequencer")));
+									}
+									else
+									{
+										FloatCurve.Name.DisplayName = CurveElement->GetName();	
+									}						
+								
+									Skeleton->VerifySmartName(USkeleton::AnimCurveMappingName, FloatCurve.Name);
+									FloatCurve.Color = Element->Settings.ShapeColor;
+									
+									const FAnimationCurveIdentifier CurveId(FloatCurve.Name, ERawCurveTrackTypes::RCT_Float);
+									if (!RetainFloatCurves || !FloatCurve.Name.DisplayName.ToString().Contains(SequencerSuffix))
+									{
+										const FAnimationCurveMetaData& CurveMetaData = CurveIdentifierToMetaData.FindChecked(CurveId);
+										FloatCurve.SetCurveTypeFlags(CurveMetaData.Flags);
+										FloatCurve.Color = CurveMetaData.Color;
+									}							
+
+									if (const FScalarParameterNameAndCurve* ScalarCurve = ScalarCurves.FindByPredicate([Element](FScalarParameterNameAndCurve Curve)
+									{
+										return Curve.ParameterName == Element->GetName();
+									}))
+									{							
+										AnimSequencerHelpers::ConvertFloatChannelToRichCurve(ScalarCurve->ParameterCurve, FloatCurve.FloatCurve, FrameRate);
+									}
+								}
+								return true;
+							});	
+						}
+						else
+						{						
+							IAnimationDataController::ReportObjectErrorf(this, LOCTEXT("UnableToFindRigHierarchy", "Unable to retrieve RigHierarchy for ControlRig ({0})"), FText::FromString(ControlRig->GetPathName()));	      
+						}
+					}								
+				}
+			}
+		}
+	}
+}
+
+void UAnimationSequencerDataModel::GenerateLegacyBoneData()
+{
+	QUICK_SCOPE_CYCLE_COUNTER(STAT_GenerateLegacyBoneData);
+	// Reset current track-data
+	LegacyBoneAnimationTracks.Reset();
+	
+	if(const USkeleton* TargetSkeleton = GetSkeleton())
+	{
+		const FReferenceSkeleton& ReferenceSkeleton = TargetSkeleton->GetReferenceSkeleton();
+		ValidateSequencerData();
+		
+		if(const UMovieSceneControlRigParameterSection* Section = GetFKControlRigSection())
+		{
+			if (const UControlRig* ControlRig = Section->GetControlRig())
+			{
+				if (const URigHierarchy* Hierarchy = ControlRig->GetHierarchy())
+				{
+					TArray<FTransform> Transforms;
+					TArray<FFrameNumber> FrameNumbers;
+
+					// Loop over all control elements in hierarchy - and for any found bone controllers sample their curve data to a new track entry
+					Hierarchy->ForEach<FRigControlElement>([this, Section, ReferenceSkeleton, TargetSkeleton, Hierarchy, &Transforms, &FrameNumbers](const FRigControlElement* ControlElement) -> bool
+					{
+						const FName TargetBoneName = UFKControlRig::GetControlTargetName(ControlElement->GetName(), ERigElementType::Bone);
+						const FRigElementKey BoneElementKey(TargetBoneName, ERigElementType::Bone);
+
+						if (Hierarchy->Contains(BoneElementKey))
+						{
+							const FTransformParameterNameAndCurves* ParameterCurvePair = Section->GetTransformParameterNamesAndCurves().FindByPredicate([ControlName = ControlElement->GetName()](const FTransformParameterNameAndCurves Parameter)
+							{
+								return Parameter.ParameterName == ControlName;
+							});
+
+							// Only populate the track if any curve keys were set
+							const bool bContainsAnyKeys = ParameterCurvePair && [ParameterCurvePair]()
+							{
+								for (int32 ChannelIndex = 0; ChannelIndex < 3; ++ChannelIndex)
+								{
+									if (ParameterCurvePair->Translation[ChannelIndex].GetNumKeys())
+										return true;
+									if (ParameterCurvePair->Rotation[ChannelIndex].GetNumKeys())
+										return true;
+									if (ParameterCurvePair->Scale[ChannelIndex].GetNumKeys())
+										return true;				
+								}
+
+								return false;
+							}();
+							
+							FBoneAnimationTrack& BoneTrack = LegacyBoneAnimationTracks.AddDefaulted_GetRef();
+							BoneTrack.Name = TargetBoneName;
+							BoneTrack.BoneTreeIndex = ReferenceSkeleton.FindBoneIndex(BoneTrack.Name);
+							ensureMsgf(BoneTrack.BoneTreeIndex != INDEX_NONE, TEXT("Generating track data for a bone %s that does not exist on the skeleton %"), *TargetBoneName.ToString(), *TargetSkeleton->GetPathName());
+							
+							if(bContainsAnyKeys)
+							{
+								Transforms.Reset();
+								FrameNumbers.Reset();
+								
+								GenerateTransformKeysForControl(TargetBoneName, Transforms, FrameNumbers);
+								for (const FTransform& Transform : Transforms)
+								{
+									BoneTrack.InternalTrackData.PosKeys.Add(FVector3f(Transform.GetLocation()));
+									BoneTrack.InternalTrackData.RotKeys.Add(FQuat4f(Transform.GetRotation()));
+									BoneTrack.InternalTrackData.ScaleKeys.Add(FVector3f(Transform.GetScale3D()));					
+								}
+							}
+						}
+						return true;
+					});
+				}
+				else
+				{						
+					IAnimationDataController::ReportObjectErrorf(this, LOCTEXT("UnableToFindRigHierarchy", "Unable to retrieve RigHierarchy for ControlRig ({0})"), FText::FromString(ControlRig->GetPathName()));	      
+				}
+			}
+			else
+			{				
+				IAnimationDataController::ReportObjectErrorf(this, LOCTEXT("UnableToFindControlRig", "Unable to retrieve ControlRig for Model ({0})"), FText::FromString(GetPathName()));
+			}
+		}					
+	}
+	else
+	{
+		IAnimationDataController::ReportObjectErrorf(this, LOCTEXT("UnableToFindSkeleton", "Unable to retrieve target USkeleton for Animation Asset ({0})"), FText::FromString(GetOuter()->GetPathName()));	
+	}
+}
+
+void UAnimationSequencerDataModel::ValidateData() const
+{		
+	ValidateSequencerData();
+	ValidateControlRigData();
+
+	if (!!ValidationMode)
+	{
+		ValidateLegacyAgainstControlRigData();
+	}
+}
+
+void UAnimationSequencerDataModel::ValidateSequencerData() const
+{
+	checkf(MovieScene, TEXT("No Movie Scene found for SequencerDataModel"));
+
+	const int32 NumberOfMasterTracks = MovieScene->GetMasterTracks().Num();
+	checkf(NumberOfMasterTracks == 1, TEXT("Invalid number of Master Tracks in Movie Scene expected 1 but found %i"), NumberOfMasterTracks);
+		
+	const UMovieSceneControlRigParameterTrack* Track = MovieScene->FindMasterTrack<UMovieSceneControlRigParameterTrack>();
+	checkf(Track, TEXT("Unable to find Control Rig Master Track"));
+
+	const int32 NumberOfSections = Track->GetAllSections().Num();
+	checkf(NumberOfSections == 1, TEXT("Invalid number of Sections found for Control Rig Master Track expected 1 but found %i"), NumberOfSections);
+
+	UMovieSceneControlRigParameterSection* Section = GetFKControlRigSection();
+	checkf(Section, TEXT("Unable to find Control Rig Master Section"));
+}
+
+void UAnimationSequencerDataModel::ValidateControlRigData() const
+{
+	UMovieSceneControlRigParameterSection* Section = GetFKControlRigSection();
+	checkf(Section, TEXT("Unable to find Control Rig Master Section"));
+
+	UControlRig* ControlRig = Section->GetControlRig();
+	checkf(ControlRig, TEXT("Unable to find Control Rig instance for Section"));
+
+	checkf(ControlRig->IsA<UFKControlRig>(), TEXT("Invalid class for Control Rig expected UFKControlRig"));
+
+	const URigHierarchy* Hierarchy = ControlRig->GetHierarchy();
+	if (Hierarchy && ValidationMode)
+	{
+		// Validate Rig Hierarchy against the outer Animation Sequence its (reference) Skeleton		
+		if (const USkeleton* Skeleton = GetSkeleton())
+		{
+			const FReferenceSkeleton& ReferenceSkeleton = Skeleton->GetReferenceSkeleton();
+			const int32 NumberOfBones = ReferenceSkeleton.GetNum();
+
+			// Validating the bone elements against the reference skeleton bones
+			for (int32 BoneIndex = 0; BoneIndex < NumberOfBones; ++BoneIndex)
+			{
+				const FName ExpectedBoneName = ReferenceSkeleton.GetBoneName(BoneIndex);
+				const bool bIsVirtualBone = ExpectedBoneName.ToString().StartsWith(VirtualBoneNameHelpers::VirtualBonePrefix);
+				if (!bIsVirtualBone)
+				{
+					const FRigElementKey BoneKey(ExpectedBoneName, ERigElementType::Bone);
+					const FRigBoneElement* BoneElement = Hierarchy->Find<FRigBoneElement>(BoneKey);
+					checkf(BoneElement, TEXT("Unable to find FRigBoneElement in RigHierarchy for Bone with name: %s"), *ExpectedBoneName.ToString());
+		
+					const int32 ParentBoneIndex = ReferenceSkeleton.GetParentIndex(BoneIndex);
+					if (BoneElement && ParentBoneIndex != INDEX_NONE)
+					{
+						const FName ExpectedParentBoneName = ReferenceSkeleton.GetBoneName(ParentBoneIndex);
+						const FRigElementKey ParentBoneKey(ExpectedParentBoneName, ERigElementType::Bone);
+            
+						const FRigBoneElement* ParentBoneElement = Hierarchy->Find<FRigBoneElement>(ParentBoneKey);
+						checkf(BoneElement->ParentElement == ParentBoneElement, TEXT("Unexpected Parent Element for Bone %s. Expected %s but found %s"), *ExpectedBoneName.ToString(), *ExpectedParentBoneName.ToString(), *ParentBoneElement->GetDisplayName().ToString());
+					}	
+				}
+			}
+		}
+	}
+}
+
+void UAnimationSequencerDataModel::ValidateLegacyAgainstControlRigData() const
+{
+	UMovieSceneControlRigParameterSection* Section = GetFKControlRigSection();
+
+	UControlRig* ControlRig = Section->GetControlRig();
+	URigHierarchy* Hierarchy = ControlRig->GetHierarchy();
+
+	// Validate bone tracks against controls
+	const UAnimSequence* OuterSequence = GetAnimationSequence();
+	if (const USkeleton* Skeleton = OuterSequence->GetSkeleton())
+	{
+		const FReferenceSkeleton& ReferenceSkeleton = Skeleton->GetReferenceSkeleton();
+
+		TArray<FTransform> Transforms;
+		TArray<FFrameNumber> FrameNumbers;
+		
+		for (const FBoneAnimationTrack& Track : LegacyBoneAnimationTracks)
+		{
+			const FName ExpectedBoneName = ReferenceSkeleton.GetBoneName(Track.BoneTreeIndex);
+			
+			const FRigElementKey BoneKey(ExpectedBoneName, ERigElementType::Bone);
+			const FRigBoneElement* BoneElement = Hierarchy->Find<FRigBoneElement>(BoneKey);
+			if(!BoneElement)
+			{
+				IAnimationDataController::ReportObjectErrorf(this, LOCTEXT("BoneElementNotFound", "Unable to find FRigBoneElement in RigHierarchy for Bone with name: {0}"), FText::FromString(ExpectedBoneName.ToString()));
+			}
+
+			const FRigElementKey BoneControlKey(UFKControlRig::GetControlName(ExpectedBoneName, ERigElementType::Bone), ERigElementType::Control);
+			const FRigControlElement* BoneControlElement = Hierarchy->Find<FRigControlElement>(BoneControlKey);
+
+			if (!BoneControlElement)
+			{
+				IAnimationDataController::ReportObjectErrorf(this, LOCTEXT("ControlElementNotFound", "Unable to find FRigControlElement in RigHierarchy for Bone with name: {0}"), FText::FromString(ExpectedBoneName.ToString()));
+			}		
+		
+			FTransformParameterNameAndCurves* BoneCurveParameter = Section->GetTransformParameterNamesAndCurves().FindByPredicate([BoneControlKey](const FTransformParameterNameAndCurves& ParameterPair)
+			{
+				return ParameterPair.ParameterName == BoneControlKey.Name;
+			});
+			if (!(BoneCurveParameter || Track.InternalTrackData.PosKeys.Num() == 0))
+			{
+				IAnimationDataController::ReportObjectErrorf(this, LOCTEXT("ControlCurveNotFound", "Unable to find FTransformParameterNameAndCurves in RigHierarchy for Bone Control with name: {0}"), FText::FromName(BoneControlKey.Name));
+			}
+
+			GenerateTransformKeysForControl(ExpectedBoneName, Transforms, FrameNumbers);
+
+			const int32 NumExpectedKeys = Track.InternalTrackData.PosKeys.Num();
+			if (NumExpectedKeys != Transforms.Num())
+			{
+				IAnimationDataController::ReportObjectErrorf(this, LOCTEXT("UnexpectedNumberOfControlKeys",	"Unexpected number of Bone Control Curve keys for {0}, expected {1} but found {2}"), FText::FromName(ExpectedBoneName), FText::AsNumber(NumExpectedKeys), FText::AsNumber(Transforms.Num()));
+			}			
+
+			if (NumExpectedKeys == Transforms.Num())
+			{
+				for (int32 KeyIndex = 0; KeyIndex < NumExpectedKeys; ++KeyIndex)
+				{
+					const FTransform& TransformKey = Transforms[KeyIndex];
+
+					checkf(TransformKey.GetLocation().Equals(FVector(Track.InternalTrackData.PosKeys[KeyIndex])), TEXT("Unexpected positional key (%i) for bone %s, expected %s but found %s"), KeyIndex,
+						*ExpectedBoneName.ToString(),* Track.InternalTrackData.PosKeys[KeyIndex].ToCompactString(), *TransformKey.GetLocation().ToCompactString());
+
+					if(!TransformKey.GetLocation().Equals(FVector(Track.InternalTrackData.PosKeys[KeyIndex])))
+					{
+						IAnimationDataController::ReportObjectErrorf(this, LOCTEXT("UnexpectedPositionalKey", "Unexpected positional key ({0}) for bone {1}, expected {2} but found {3}"),							
+							FText::AsNumber(KeyIndex), FText::FromName(ExpectedBoneName), FText::FromString(Track.InternalTrackData.PosKeys[KeyIndex].ToCompactString()), FText::FromString(TransformKey.GetLocation().ToCompactString()));
+					}
+
+					const FQuat LegacyRotation = FQuat(Track.InternalTrackData.RotKeys[KeyIndex]).GetNormalized();
+					const double RotationDeltaDegrees = FMath::RadiansToDegrees(TransformKey.GetRotation().AngularDistance(LegacyRotation));
+			
+					if(RotationDeltaDegrees > 0.5f)
+					{
+						IAnimationDataController::ReportObjectErrorf(this, LOCTEXT("UnexpectedRotationalKey", "Unexpected rotational key ({0} degrees delta) ({1}) for bone {2}, expected {3} but found {4}"), FText::AsNumber(RotationDeltaDegrees), FText::AsNumber(KeyIndex), FText::FromName(ExpectedBoneName), FText::FromString(Track.InternalTrackData.RotKeys[KeyIndex].ToString()), FText::FromString(TransformKey.GetRotation().ToString()));
+					}
+
+					if (!TransformKey.GetScale3D().Equals(FVector(Track.InternalTrackData.ScaleKeys[KeyIndex])))
+					{
+						IAnimationDataController::ReportObjectErrorf(this, LOCTEXT("UnexpectedScalingKey", "Unexpected scaling key ({0}) for bone {1}, expected {2} but found {3}"), FText::AsNumber(KeyIndex),
+							FText::FromName(ExpectedBoneName), FText::FromString(Track.InternalTrackData.ScaleKeys[KeyIndex].ToCompactString()), FText::FromString(TransformKey.GetScale3D().ToCompactString()));
+					}
+				}
+			
+				Transforms.Reset();
+				FrameNumbers.Reset();
+			}
+		}
+		
+		// Validate curve data against controls
+		for (const FFloatCurve& FloatCurve : LegacyCurveData.FloatCurves)
+		{
+			const FName CurveName = FloatCurve.Name.DisplayName;					
+			const FRigElementKey CurveKey(CurveName, ERigElementType::Curve);
+			const FRigCurveElement* CurveElement = Hierarchy->Find<FRigCurveElement>(CurveKey);
+			if (!CurveElement)
+			{
+				IAnimationDataController::ReportObjectErrorf(this, LOCTEXT("CurveElementNotFound", "Unable to find FRigCurve in RigHierarchy for Curve with name: {0}"), FText::FromName(CurveName));
+			}
+	        
+	        const FRigElementKey CurveControlKey(UFKControlRig::GetControlName(CurveName, ERigElementType::Curve), ERigElementType::Control);
+	        const FRigControlElement* CurveControlElement = Hierarchy->Find<FRigControlElement>(CurveControlKey);
+			if (!CurveControlElement)
+			{
+				IAnimationDataController::ReportObjectErrorf(this, LOCTEXT("CurveControlElementNotFound", "Unable to find FRigControlElement in RigHierarchy for Curve with name: {0}"), FText::FromName(CurveName));
+			}
+			
+			FScalarParameterNameAndCurve* CurveControlParameter = Section->GetScalarParameterNamesAndCurves().FindByPredicate([CurveControlKey](const FScalarParameterNameAndCurve& ParameterPair)
+			{
+				return ParameterPair.ParameterName == CurveControlKey.Name;
+			});
+			
+			if (CurveControlParameter)
+			{
+				for (const FRichCurveKey& Key : FloatCurve.FloatCurve.GetConstRefOfKeys())
+				{
+					float ParameterValue = 0.f;
+					const FFrameTime FrameTime = CurveControlParameter->ParameterCurve.GetTickResolution().AsFrameTime(Key.Time);
+
+					if (!CurveControlParameter->ParameterCurve.Evaluate(FrameTime, ParameterValue))
+					{
+						IAnimationDataController::ReportObjectErrorf(this, LOCTEXT("FailedToEvaluateCurveControl", "Unable to evaluate Control Curve ({0}) at interval {1}"), FText::FromName(CurveName), FText::AsNumber(FrameTime.AsDecimal()));
+					}
+
+					const float RichCurveValue = FloatCurve.FloatCurve.Eval(Key.Time);
+					// QQ threshold
+					if (!(FMath::IsNearlyEqual(ParameterValue, Key.Value, 0.001f) || FMath::IsNearlyEqual(ParameterValue, RichCurveValue, 0.001f)))
+					{
+						IAnimationDataController::ReportObjectErrorf(this, LOCTEXT("CurveDeviationError", "Unexpected Control Curve ({0}) evaluation value {1} at {2}, expected {3} ({4})"), FText::FromName(CurveName), FText::AsNumber(ParameterValue), FText::AsNumber(FrameTime.AsDecimal()), FText::AsNumber(Key.Value), FText::AsNumber(RichCurveValue));
+					}
+				}
+			}
+			else
+			{
+				IAnimationDataController::ReportObjectErrorf(this, LOCTEXT("ParameterNotFound", "Unable to find FScalarParameterNameAndCurve in RigHierarchy for Curve Control with name: {0}"), FText::FromName(CurveName));
+			}
+		}	
+	}	
+}
+
+void UAnimationSequencerDataModel::IterateTransformControlCurve(const FName& BoneName, TFunction<void(const FTransform&, const FFrameNumber&)> IterationFunction) const
+{
+	ValidateSequencerData();
+	ValidateControlRigData();
+
+	const UMovieSceneControlRigParameterSection* Section = GetFKControlRigSection();
+	UControlRig* ControlRig = Section->GetControlRig();
+	const URigHierarchy* Hierarchy = ControlRig->GetHierarchy();
+	
+	const FRigElementKey BoneControlKey(UFKControlRig::GetControlName(BoneName, ERigElementType::Bone), ERigElementType::Control);
+	if (Hierarchy->Contains(BoneControlKey))
+	{
+		if (const FTransformParameterNameAndCurves* ControlCurvePtr = Section->GetTransformParameterNamesAndCurves().FindByPredicate([CurveName = BoneControlKey.Name](const FTransformParameterNameAndCurves& TransformParameter)
+		{
+			return TransformParameter.ParameterName == CurveName;
+		}))
+		{
+			const FTransformParameterNameAndCurves& ControlCurve = *ControlCurvePtr;
+
+			TSet<FFrameNumber> Frames;
+			AnimSequencerHelpers::GenerateUniqueFramesForTransformCurve(ControlCurve, Frames);
+
+			for (const FFrameNumber& Frame : Frames)
+			{
+				FTransform Transform;
+
+				FVector3f Location;
+				FVector3f EulerAngles;
+				FVector3f Scale;
+				
+				for (int32 ChannelIndex = 0; ChannelIndex < 3; ++ChannelIndex)
+				{
+					ControlCurve.Translation[ChannelIndex].Evaluate(Frame, Location[ChannelIndex]);
+					ControlCurve.Rotation[ChannelIndex].Evaluate(Frame, EulerAngles[ChannelIndex]);
+					ControlCurve.Scale[ChannelIndex].Evaluate(Frame, Scale[ChannelIndex]);
+				}
+
+				Transform.SetLocation(FVector(Location));
+				Transform.SetRotation(FQuat::MakeFromEuler(FVector(EulerAngles)));
+				Transform.SetScale3D(FVector(Scale));
+
+				Transform.NormalizeRotation();
+
+				IterationFunction(Transform, Frame);
+			}
+		}
+	}
+}
+
+void UAnimationSequencerDataModel::GenerateTransformKeysForControl(const FName& BoneName, TArray<FTransform>& InOutTransforms, TArray<FFrameNumber>& InOutFrameNumbers) const
+{
+	IterateTransformControlCurve(BoneName, [&InOutTransforms, &InOutFrameNumbers](const FTransform& Transform, const FFrameNumber& FrameNumber) -> void
+	{
+		InOutTransforms.Add(Transform);
+		InOutFrameNumbers.Add(FrameNumber);
+	});
+}
+
+UMovieScene* UAnimationSequencerDataModel::GetMovieScene() const
+{
+	return MovieScene;
+}
+
+UObject* UAnimationSequencerDataModel::GetParentObject(UObject* MovieSceneBlends) const
+{
+	return GetOuter();
+}
+
+void UAnimationSequencerDataModel::GeneratePoseData(UControlRig* ControlRig, FAnimationPoseData& InOutPoseData, const UE::Anim::DataModel::FEvaluationContext& EvaluationContext) const
+{
+	QUICK_SCOPE_CYCLE_COUNTER(STAT_GeneratePoseData);
+	
+	if (ControlRig)
+	{
+		if (const URigHierarchy* RigHierarchy = ControlRig->GetHierarchy())
+		{
+			// Evaluate Control rig to update bone and curve elements according to controls
+			ControlRig->Evaluate_AnyThread();
+
+			// Start with ref-pose
+			FCompactPose& RigPose = InOutPoseData.GetPose();
+			RigPose.ResetToRefPose();
+			const FBoneContainer& RequiredBones = RigPose.GetBoneContainer();
+
+			FCSPose<FCompactPose> MeshPoses;
+			MeshPoses.InitPose(RigPose);
+
+			FBlendedCurve& Curve = InOutPoseData.GetCurve();
+			UE::Anim::Retargeting::FRetargetingScope RetargetingScope(RigPose, EvaluationContext);
+			
+			// Populate bone/curve elements to Pose/Curve indices
+			{
+				QUICK_SCOPE_CYCLE_COUNTER(STAT_GetMappings);
+				const FReferenceSkeleton& RefSkeleton = RequiredBones.GetReferenceSkeleton();
+
+				RigHierarchy->ForEach<FRigControlElement>([this, &RefSkeleton, &Curve, &RequiredBones, &RigHierarchy, &RetargetingScope, &MeshPoses, bValidCurve = Curve.IsValid(), SmartNameContainer=RequiredBones.GetSkeletonAsset()->GetSmartNameContainer(USkeleton::AnimCurveMappingName)](const FRigControlElement* ControlElement) -> bool
+				{
+					if (ControlElement->Settings.ControlType == ERigControlType::EulerTransform)
+					{
+						const FName ExpectedBoneName = UFKControlRig::GetControlTargetName(ControlElement->GetName(), ERigElementType::Bone);
+						const int32 BoneIndex = RefSkeleton.FindBoneIndex(ExpectedBoneName);
+						if (BoneIndex != INDEX_NONE)
+						{
+							const FName& BoneName = ExpectedBoneName;
+							const FRigElementKey Key(BoneName, ERigElementType::Bone);
+
+							const bool bMatchingLegacyBone = !ValidationMode || LegacyBoneAnimationTracks.ContainsByPredicate([BoneName](const FBoneAnimationTrack& Track)
+							{
+								return Track.Name == BoneName;
+							});
+							ensureMsgf(bMatchingLegacyBone, TEXT("Non-matching bone vs legacy data %s"), *BoneName.ToString());
+						
+							const int32 SkeletonBoneIndex = RequiredBones.GetSkeletonAsset()->GetReferenceSkeleton().FindBoneIndex(BoneName);
+							const FCompactPoseBoneIndex CompactPoseBoneIndex = RequiredBones.GetCompactPoseIndexFromSkeletonIndex(SkeletonBoneIndex);
+							if (CompactPoseBoneIndex != INDEX_NONE)
+							{
+								RetargetingScope.AddTrackedBone(CompactPoseBoneIndex, SkeletonBoneIndex);
+				        
+								// Retrieve evaluated bone transform from Hierarchy
+								const FTransform ComponentTransform = RigHierarchy->GetGlobalTransform(Key);
+								MeshPoses.SetComponentSpaceTransform(CompactPoseBoneIndex, ComponentTransform);
+							}
+						}
+						
+					}
+					else if (ControlElement->Settings.ControlType == ERigControlType::Float && bValidCurve)
+					{
+						const FName ExpectedCurveName = UFKControlRig::GetControlTargetName(ControlElement->GetName(), ERigElementType::Curve);
+						const SmartName::UID_Type CurveIndex = SmartNameContainer->FindUID(ExpectedCurveName);
+						if (CurveIndex != INDEX_NONE)
+						{
+							const FRigElementKey Key(ExpectedCurveName, ERigElementType::Curve);	
+							if (Curve.IsEnabled(CurveIndex))
+							{
+								Curve.Set(CurveIndex, RigHierarchy->GetCurveValue(Key));
+							}
+							
+						}
+					}
+
+					return true;
+				});
+			}			
+
+			{
+				QUICK_SCOPE_CYCLE_COUNTER(STAT_ConvertToLocalTransforms);
+				// Convert back to local-space bone transforms
+				FCSPose<FCompactPose>::ConvertComponentPosesToLocalPosesSafe(MeshPoses, RigPose);
+				RigPose.NormalizeRotations();
+			}
+
+			// Apply any additive transform curves - if requested and any are set
+			if (!RigPose.GetBoneContainer().ShouldUseSourceData())
+			{
+				for (const FTransformCurve& TransformCurve : GetTransformCurves())
+				{
+					// if disabled, do not handle
+					if (TransformCurve.GetCurveTypeFlag(AACF_Disabled))
+					{
+						continue;
+					}
+			
+					// Add or retrieve curve
+					const FName& CurveName = TransformCurve.Name.DisplayName;
+					// note we're not checking Curve.GetCurveTypeFlags() yet
+					FTransform Value = TransformCurve.Evaluate(EvaluationContext.SampleFrameRate.AsSeconds(EvaluationContext.SampleTime), 1.f);
+
+					const FCompactPoseBoneIndex BoneIndex(RigPose.GetBoneContainer().GetPoseBoneIndexForBoneName(CurveName));
+					ensure(BoneIndex != INDEX_NONE);
+
+					const FTransform LocalTransform = RigPose[BoneIndex];
+					RigPose[BoneIndex].SetRotation(LocalTransform.GetRotation() * Value.GetRotation());
+					RigPose[BoneIndex].SetTranslation(LocalTransform.TransformPosition(Value.GetTranslation()));
+					RigPose[BoneIndex].SetScale3D(LocalTransform.GetScale3D() * Value.GetScale3D());
+				}
+			}
+
+			// Generate relative transform for VirtualBones according to source/target
+			{					
+				QUICK_SCOPE_CYCLE_COUNTER(STAT_GenerateVirtualBones);
+				
+				TArray<FVirtualBoneCompactPoseData>& VBCompactPoseData = UE::Anim::FBuildRawPoseScratchArea::Get().VirtualBoneCompactPoseData;
+				VBCompactPoseData = RequiredBones.GetVirtualBoneCompactPoseData();
+				if (VBCompactPoseData.Num() > 0)
+				{
+					FCSPose<FCompactPose> CSPose1;
+					CSPose1.InitPose(RigPose);
+
+					for (const FVirtualBoneCompactPoseData& VB : VBCompactPoseData)
+					{
+						const FTransform Source = CSPose1.GetComponentSpaceTransform(VB.SourceIndex);
+						const FTransform Target = CSPose1.GetComponentSpaceTransform(VB.TargetIndex);
+						RigPose[VB.VBIndex] = Target.GetRelativeTransform(Source);
+					}
+				}
+			}
+
+			{				
+				QUICK_SCOPE_CYCLE_COUNTER(STAT_SetAttributes);
+				// Evaluate attributes at requested time interval
+				for (const FAnimatedBoneAttribute& Attribute : AnimatedBoneAttributes)
+				{
+					const FCompactPoseBoneIndex PoseBoneIndex = RequiredBones.GetCompactPoseIndexFromSkeletonIndex(Attribute.Identifier.GetBoneIndex());
+					// Only add attribute if the bone its tied to exists in the currently evaluated set of bones
+					if(PoseBoneIndex.IsValid())
+					{
+						UE::Anim::Attributes::GetAttributeValue(InOutPoseData.GetAttributes(), PoseBoneIndex, Attribute, EvaluationContext.SampleFrameRate.AsSeconds(EvaluationContext.SampleTime));
+					}
+				}				
+			}			
+		}
+	}
+}
+
+void UAnimationSequencerDataModel::EvaluateTrack(UMovieSceneControlRigParameterTrack* CR_Track, const UE::Anim::DataModel::FEvaluationContext& EvaluationContext) const
+{
+	QUICK_SCOPE_CYCLE_COUNTER(STAT_EvaluateTrack);
+	
+	// Determine frame-time to sample according to the interpolation type (floor to frame for step interpolation)
+	const FFrameTime InterpolationTime = EvaluationContext.InterpolationType == EAnimInterpolationType::Step ? EvaluationContext.SampleTime.FloorToFrame() : EvaluationContext.SampleTime;
+	const FFrameTime BoneSampleTime = FFrameRate::TransformTime(InterpolationTime, EvaluationContext.SampleFrameRate, MovieScene->GetTickResolution());	
+
+	// Retrieve section withing range of requested evaluation frame 
+	const TArray<UMovieSceneSection*, TInlineAllocator<4>> SectionsInRange = CR_Track->FindAllSections(BoneSampleTime.FrameNumber);
+	//ensureMsgf(SectionsInRange.Num() == 1, TEXT("Unable to retrieve section within range of request evaluation frame %i for %s"), BoneSampleTime.FrameNumber.Value, *GetAnimationSequence()->GetPathName());
+	if (SectionsInRange.Num())
+	{
+		const UMovieSceneControlRigParameterSection* FKRigSection = CastChecked<UMovieSceneControlRigParameterSection>(SectionsInRange[0]);
+		checkf(FKRigSection->ControlRigClass->GetDefaultObject()->IsA<UFKControlRig>(), TEXT("Unexpected class %s on ControlRig, expecting FKControlRig"), *FKRigSection->ControlRigClass->GetPathName());
+		
+		bool bWasDoNotKey = false;
+		bWasDoNotKey = FKRigSection->GetDoNotKey();
+		FKRigSection->SetDoNotKey(true);
+
+		UControlRig* ControlRig = FKRigSection->GetControlRig();
+		check(ControlRig);
+
+		// Reset to ref-pose
+		ControlRig->GetHierarchy()->ResetPoseToInitial(ERigElementType::Bone);
+
+		const TArray<FScalarParameterNameAndCurve>& ScalarParameters = FKRigSection->GetScalarParameterNamesAndCurves();
+		for (const FScalarParameterNameAndCurve& TypedParameter : ScalarParameters)
+		{
+			const FName& Name = TypedParameter.ParameterName;
+			float Value = 0.f;
+
+			const FFrameTime CurveSampleTime = FFrameRate::TransformTime(EvaluationContext.SampleTime, EvaluationContext.SampleFrameRate, TypedParameter.ParameterCurve.GetTickResolution());
+			if(TypedParameter.ParameterCurve.Evaluate(CurveSampleTime, Value))
+			{					
+				const FRigControlElement* ControlElement = ControlRig->FindControl(Name);
+				if (ControlElement && ControlElement->Settings.ControlType == ERigControlType::Float)
+				{
+					ControlRig->SetControlValue<float>(Name, Value, false, EControlRigSetKey::Never, false);
+				}
+			}
+		}
+
+		const TArray<FTransformParameterNameAndCurves>& TransformParameters = FKRigSection->GetTransformParameterNamesAndCurves();
+		if (TransformParameters.Num())
+		{
+
+			struct FEvaluationInfo
+			{
+				float Interp = 0.f;
+				int32 Index1 = INDEX_NONE, Index2 = INDEX_NONE;
+			};	
+
+			FEvaluationInfo ZeroAlphaInfo;
+			UE::MovieScene::EvaluateTime(TransformParameters[0].Translation[0].GetTimes(), BoneSampleTime.FrameNumber, ZeroAlphaInfo.Index1, ZeroAlphaInfo.Index2, ZeroAlphaInfo.Interp);
+			
+			FEvaluationInfo BlendFromInfo;
+			UE::MovieScene::EvaluateTime(TransformParameters[0].Translation[0].GetTimes(), BoneSampleTime.FloorToFrame(), BlendFromInfo.Index1, BlendFromInfo.Index2, BlendFromInfo.Interp);
+
+			FEvaluationInfo BlendToInfo;
+			UE::MovieScene::EvaluateTime(TransformParameters[0].Translation[0].GetTimes(), BoneSampleTime.CeilToFrame(), BlendToInfo.Index1, BlendToInfo.Index2, BlendToInfo.Interp);
+			
+			for (const FTransformParameterNameAndCurves& TypedParameter : TransformParameters)
+			{
+				const FName& Name = TypedParameter.ParameterName;
+				const FRigControlElement* ControlElement = ControlRig->FindControl(Name);
+				if (ControlElement && ControlElement->Settings.ControlType == ERigControlType::EulerTransform)
+				{
+					FEulerTransform EulerTransform;
+					
+					const double Alpha = BoneSampleTime.GetSubFrame();			
+					auto EvaluateToTransform = [&TypedParameter](const FFrameNumber Frame) -> FTransform
+					{
+						FVector3f Translation(ForceInitToZero), Scale(FVector3f::OneVector);
+						FRotator3f Rotator(0.0f, 0.0f, 0.0f);
+
+						TypedParameter.Translation[0].Evaluate(Frame,Translation[0]);
+						TypedParameter.Translation[1].Evaluate(Frame, Translation[1]);
+						TypedParameter.Translation[2].Evaluate(Frame, Translation[2]);
+
+						TypedParameter.Scale[0].Evaluate(Frame, Scale[0]);
+						TypedParameter.Scale[1].Evaluate(Frame, Scale[1]);
+						TypedParameter.Scale[2].Evaluate(Frame, Scale[2]);
+						
+						TypedParameter.Rotation[0].Evaluate(Frame, Rotator.Roll);
+						TypedParameter.Rotation[1].Evaluate(Frame, Rotator.Pitch);
+						TypedParameter.Rotation[2].Evaluate(Frame, Rotator.Yaw);
+
+						return FTransform(FRotator(Rotator), static_cast<FVector>(Translation), static_cast<FVector>(Scale));
+					};
+
+					auto ExtractTransform = [&TypedParameter](const FFrameNumber Frame) -> FTransform
+					{
+						FVector3f Translation(ForceInitToZero), Scale(FVector3f::OneVector);	
+
+						Translation[0] = TypedParameter.Translation[0].GetValues()[Frame.Value].Value;
+						Translation[1] = TypedParameter.Translation[1].GetValues()[Frame.Value].Value;
+						Translation[2] = TypedParameter.Translation[2].GetValues()[Frame.Value].Value;
+
+						Scale[0] = TypedParameter.Scale[0].GetValues()[Frame.Value].Value;
+						Scale[1] = TypedParameter.Scale[1].GetValues()[Frame.Value].Value;
+						Scale[2] = TypedParameter.Scale[2].GetValues()[Frame.Value].Value;
+
+						FRotator3f Rotator;
+						Rotator.Roll = TypedParameter.Rotation[0].GetValues()[Frame.Value].Value;
+						Rotator.Pitch = TypedParameter.Rotation[1].GetValues()[Frame.Value].Value;
+						Rotator.Yaw = TypedParameter.Rotation[2].GetValues()[Frame.Value].Value;
+
+						return FTransform(FRotator(Rotator), static_cast<FVector>(Translation), static_cast<FVector>(Scale));
+					};
+					
+					// Assume no interpolation due to uniform keys
+					if (FMath::IsNearlyZero(Alpha))
+					{
+						if (EvaluationContext.InterpolationType == EAnimInterpolationType::Linear)
+						{
+							EulerTransform = FEulerTransform(EvaluateToTransform(BoneSampleTime.FrameNumber));
+						}
+						else if (EvaluationContext.InterpolationType == EAnimInterpolationType::Step)
+						{
+							EulerTransform = FEulerTransform(ExtractTransform(BoneSampleTime.FrameNumber));
+						}
+					}
+					// Interpolate between two uniform keys
+					else
+					{
+						const FFrameNumber FromFrame = BoneSampleTime.FloorToFrame();	
+						const FFrameNumber ToFrame = BoneSampleTime.CeilToFrame();
+
+						const FTransform FromBoneTransform = EvaluateToTransform(FromFrame);
+						const FTransform ToBoneTransform = EvaluateToTransform(ToFrame);
+
+						FTransform FinalTransform;
+						FinalTransform.Blend(FromBoneTransform, ToBoneTransform, Alpha);
+						
+						EulerTransform = FEulerTransform(FinalTransform);
+					}
+					
+					ControlRig->SetControlValue<FRigControlValue::FEulerTransform_Float>(Name, EulerTransform, false, EControlRigSetKey::Never, false);
+				}
+			}
+		}
+
+		FKRigSection->SetDoNotKey(bWasDoNotKey);
+	}
+}
+
+FTransformCurve* UAnimationSequencerDataModel::FindMutableTransformCurveById(const FAnimationCurveIdentifier& CurveIdentifier)
+{
+	for (FTransformCurve& TransformCurve : LegacyCurveData.TransformCurves)
+	{
+		if (TransformCurve.Name.UID == CurveIdentifier.InternalName.UID)
+		{
+			return &TransformCurve;
+		}
+	}
+
+	return nullptr;
+}
+
+FFloatCurve* UAnimationSequencerDataModel::FindMutableFloatCurveById(const FAnimationCurveIdentifier& CurveIdentifier)
+{
+	for (FFloatCurve& FloatCurve : LegacyCurveData.FloatCurves)
+	{
+		if (FloatCurve.Name.UID == CurveIdentifier.InternalName.UID)
+		{
+			return &FloatCurve;
+		}
+	}
+
+	return nullptr;
+}
+
+FAnimCurveBase* UAnimationSequencerDataModel::FindMutableCurveById(const FAnimationCurveIdentifier& CurveIdentifier)
+{
+	switch (CurveIdentifier.CurveType)
+	{
+	case ERawCurveTrackTypes::RCT_Float:
+		return FindMutableFloatCurveById(CurveIdentifier);
+	case ERawCurveTrackTypes::RCT_Transform:
+		return FindMutableTransformCurveById(CurveIdentifier);
+	default:
+		checkf(false, TEXT("Invalid curve identifier type"));
+	}
+
+	return nullptr;
+}
+
+FRichCurve* UAnimationSequencerDataModel::GetMutableRichCurve(const FAnimationCurveIdentifier& CurveIdentifier)
+{
+	FRichCurve* RichCurve = nullptr;
+
+	if (CurveIdentifier.IsValid())
+	{
+		if (CurveIdentifier.CurveType == ERawCurveTrackTypes::RCT_Float)
+		{
+			FFloatCurve* Curve = FindMutableFloatCurveById(CurveIdentifier);
+			if (Curve)
+			{
+				RichCurve = &Curve->FloatCurve;
+			}
+		}
+		else if (CurveIdentifier.CurveType == ERawCurveTrackTypes::RCT_Transform)
+		{
+			if (CurveIdentifier.Channel != ETransformCurveChannel::Invalid && CurveIdentifier.Axis != EVectorCurveChannel::Invalid)
+			{
+				// Dealing with transform curve
+				if (FTransformCurve* TransformCurve = FindMutableTransformCurveById(CurveIdentifier))
+				{
+					if (FVectorCurve* VectorCurve = TransformCurve->GetVectorCurveByIndex(static_cast<int32>(CurveIdentifier.Channel)))
+					{
+						RichCurve = &VectorCurve->FloatCurves[static_cast<int32>(CurveIdentifier.Axis)];
+					}
+				}
+
+			}
+		}
+	}
+
+	return RichCurve;
+}
+
+#undef LOCTEXT_NAMESPACE //"AnimSequencerDataModel"
\ No newline at end of file
diff --git a/Engine/Plugins/Animation/AnimationData/Source/AnimationData/Private/AnimSequencerDataModel.h b/Engine/Plugins/Animation/AnimationData/Source/AnimationData/Private/AnimSequencerDataModel.h
new file mode 100644
--- /dev/null
+++ b/Engine/Plugins/Animation/AnimationData/Source/AnimationData/Private/AnimSequencerDataModel.h
@@ -0,0 +1,168 @@
+// Copyright Epic Games, Inc. All Rights Reserved.
+
+#pragma once
+
+#include "MovieSceneSequence.h"
+#include "Animation/AnimData/IAnimationDataModel.h"
+#include "Sections/MovieSceneParameterSection.h"
+#include "Animation/AnimData/AnimDataModelNotifyCollector.h"
+
+#include "AnimSequencerDataModel.generated.h"
+
+class UControlRig;
+class UMovieSceneControlRigParameterTrack;
+class UMovieSceneControlRigParameterSection;
+class UFKControlRig;
+
+USTRUCT()
+struct FAnimationCurveMetaData
+{
+	GENERATED_BODY()
+
+	UPROPERTY()
+	int32 Flags = 0;
+
+	UPROPERTY()
+	FLinearColor Color = FLinearColor::White;
+};
+
+UCLASS()
+class ANIMATIONDATA_API UAnimationSequencerDataModel : public UMovieSceneSequence, public IAnimationDataModel
+{
+private:
+	GENERATED_BODY()
+public:
+	static int32 RetainFloatCurves;
+	static int32 ValidationMode;
+	static int32 UseDirectFKControlRigMode;
+	
+	/** Begin UObject overrides*/
+	virtual void GetPreloadDependencies(TArray<UObject*>& OutDeps) override;
+	virtual void PostDuplicate(bool bDuplicateForPIE) override;
+	/** End UObject overrides */
+
+	/** Begin UMovieSceneSequence overrides */
+	virtual void BindPossessableObject(const FGuid& ObjectId, UObject& PossessedObject, UObject* Context) override {}
+	virtual bool CanPossessObject(UObject& Object, UObject* InPlaybackContext) const override { return false; }
+	virtual void LocateBoundObjects(const FGuid& ObjectId, UObject* Context, TArray<UObject*, TInlineAllocator<1>>& OutObjects) const override {}
+	virtual UMovieScene* GetMovieScene() const override;
+	virtual UObject* GetParentObject(UObject* Object) const override;
+	virtual void UnbindPossessableObjects(const FGuid& ObjectId) override {}
+	virtual void UnbindObjects(const FGuid& ObjectId, const TArray<UObject*>& InObjects, UObject* Context) override {}
+	virtual void UnbindInvalidObjects(const FGuid& ObjectId, UObject* Context) override {}
+	virtual void PostLoad() override;
+	/** End UMovieSceneSequence overrides */
+	
+	/** Begin IAnimationDataModel overrides*/
+	virtual double GetPlayLength() const override;
+	virtual int32 GetNumberOfFrames() const override;	
+	virtual int32 GetNumberOfKeys() const override;
+	virtual FFrameRate GetFrameRate() const override;
+	PRAGMA_DISABLE_DEPRECATION_WARNINGS	
+	virtual const TArray<FBoneAnimationTrack>& GetBoneAnimationTracks() const override;
+	virtual const FBoneAnimationTrack& GetBoneTrackByIndex(int32 TrackIndex) const override;
+	virtual const FBoneAnimationTrack& GetBoneTrackByName(FName TrackName) const override;
+	virtual const FBoneAnimationTrack* FindBoneTrackByName(FName Name) const override;
+	virtual const FBoneAnimationTrack* FindBoneTrackByIndex(int32 BoneIndex) const override;
+	virtual int32 GetBoneTrackIndex(const FBoneAnimationTrack& Track) const override;
+	virtual int32 GetBoneTrackIndexByName(FName TrackName) const override;	
+	virtual bool IsValidBoneTrackIndex(int32 TrackIndex) const override;
+	virtual int32 GetNumBoneTracks() const override;
+	virtual void GetBoneTrackNames(TArray<FName>& OutNames) const override;
+	virtual const FAnimCurveBase* FindCurve(const FAnimationCurveIdentifier& CurveIdentifier) const override;
+	virtual const FFloatCurve* FindFloatCurve(const FAnimationCurveIdentifier& CurveIdentifier) const override;
+	virtual const FTransformCurve* FindTransformCurve(const FAnimationCurveIdentifier& CurveIdentifier) const override;
+	virtual const FRichCurve* FindRichCurve(const FAnimationCurveIdentifier& CurveIdentifier) const override;
+	PRAGMA_ENABLE_DEPRECATION_WARNINGS	
+	virtual const FAnimationCurveData& GetCurveData() const override;
+	virtual int32 GetNumberOfTransformCurves() const override;
+	virtual int32 GetNumberOfFloatCurves() const override;
+	virtual const TArray<struct FFloatCurve>& GetFloatCurves() const override;
+	virtual const TArray<struct FTransformCurve>& GetTransformCurves() const override;
+	virtual const FAnimCurveBase& GetCurve(const FAnimationCurveIdentifier& CurveIdentifier) const override;
+	virtual const FFloatCurve& GetFloatCurve(const FAnimationCurveIdentifier& CurveIdentifier) const override;
+	virtual const FTransformCurve& GetTransformCurve(const FAnimationCurveIdentifier& CurveIdentifier) const override;
+	virtual const FRichCurve& GetRichCurve(const FAnimationCurveIdentifier& CurveIdentifier) const override;
+	virtual TArrayView<const FAnimatedBoneAttribute> GetAttributes() const override;
+	virtual int32 GetNumberOfAttributes() const override;
+	virtual int32 GetNumberOfAttributesForBoneIndex(const int32 BoneIndex) const override;
+	virtual void GetAttributesForBone(const FName& BoneName, TArray<const FAnimatedBoneAttribute*>& OutBoneAttributes) const override;
+	virtual const FAnimatedBoneAttribute& GetAttribute(const FAnimationAttributeIdentifier& AttributeIdentifier) const override;
+	virtual const FAnimatedBoneAttribute* FindAttribute(const FAnimationAttributeIdentifier& AttributeIdentifier) const override;
+	virtual UAnimSequence* GetAnimationSequence() const override;
+	virtual FAnimDataModelModifiedEvent& GetModifiedEvent() override { return ModifiedEvent; }
+	virtual FGuid GenerateGuid() const override;
+	virtual TScriptInterface<IAnimationDataController> GetController() override;
+	virtual IAnimationDataModel::FModelNotifier& GetNotifier() override;
+	virtual FAnimDataModelModifiedDynamicEvent& GetModifiedDynamicEvent() override { return ModifiedEventDynamic; }
+    virtual void Evaluate(FAnimationPoseData& InOutPoseData, const UE::Anim::DataModel::FEvaluationContext& EvaluationContext) const override;
+	virtual bool HasBeenPopulated() const override { return bPopulated; }
+protected:
+	virtual void OnNotify(const EAnimDataModelNotifyType& NotifyType, const FAnimDataModelNotifPayload& Payload) override final;
+	/** End IAnimationDataModel overrides */
+	
+	/** Controller helper functionality */
+	FTransformCurve* FindMutableTransformCurveById(const FAnimationCurveIdentifier& CurveIdentifier);
+	FFloatCurve* FindMutableFloatCurveById(const FAnimationCurveIdentifier& CurveIdentifier);
+	FAnimCurveBase* FindMutableCurveById(const FAnimationCurveIdentifier& CurveIdentifier);	
+	FBoneAnimationTrack* FindMutableBoneTrackByName(FName Name);	
+	FRichCurve* GetMutableRichCurve(const FAnimationCurveIdentifier& CurveIdentifier);
+
+	void GenerateLegacyCurveData();
+	void GenerateLegacyBoneData();
+
+	void ValidateData() const;
+	void ValidateSequencerData() const;
+	void ValidateControlRigData() const;
+	void ValidateLegacyAgainstControlRigData() const;
+	
+	void GeneratePoseData(UControlRig* ControlRig, FAnimationPoseData& InOutPoseData, const UE::Anim::DataModel::FEvaluationContext& EvaluationContext) const;
+	void EvaluateTrack(UMovieSceneControlRigParameterTrack* CR_Track, const UE::Anim::DataModel::FEvaluationContext& EvaluationContext) const;
+
+	UMovieSceneControlRigParameterTrack* GetControlRigTrack() const;
+	UMovieSceneControlRigParameterSection* GetFKControlRigSection() const;
+	USkeleton* GetSkeleton() const;
+	void InitializeFKControlRig(UFKControlRig* FKControlRig, USkeleton* Skeleton) const;
+	UControlRig* GetControlRig() const;
+	
+	void IterateTransformControlCurve(const FName& BoneName, TFunction<void(const FTransform&, const FFrameNumber&)> IterationFunction) const;
+	void GenerateTransformKeysForControl(const FName& BoneName, TArray<FTransform>& InOutTransforms, TArray<FFrameNumber>& InOutFrameNumbers) const;
+	void RemoveOutOfDateControls() const;
+
+private:	
+	/** Dynamic delegate event allows scripting to register to any broadcast-ed notify. */
+	UPROPERTY(BlueprintAssignable, Transient, Category = AnimationDataModel, meta = (ScriptName = "ModifiedEvent", AllowPrivateAccess = "true"))
+	FAnimDataModelModifiedDynamicEvent ModifiedEventDynamic;	
+	FAnimDataModelModifiedEvent ModifiedEvent;
+	TUniquePtr<IAnimationDataModel::FModelNotifier> Notifier;
+	UE::Anim::FAnimDataModelNotifyCollector Collector;
+
+	UPROPERTY(VisibleAnywhere, Category=AnimSequencer)
+	TArray<FBoneAnimationTrack> LegacyBoneAnimationTracks;
+
+	UPROPERTY(VisibleAnywhere, Category=AnimSequencer)
+	FAnimationCurveData LegacyCurveData;
+
+	UPROPERTY(VisibleAnywhere, Category=AnimSequencer)
+	TArray<FAnimatedBoneAttribute> AnimatedBoneAttributes;
+
+	// Movie scene instance containing FK Control rig and section representing the animation data
+	UPROPERTY(VisibleAnywhere, Category=AnimSequencer)
+	TObjectPtr<UMovieScene> MovieScene = nullptr;
+
+	// Per-curve information holding flags/color, due to be deprecated in the future
+	UPROPERTY(VisibleAnywhere, Category=AnimSequencer)
+	TMap<FAnimationCurveIdentifier, FAnimationCurveMetaData> CurveIdentifierToMetaData;
+
+	UPROPERTY(VisibleAnywhere, Category=AnimSequencer)
+	bool bPopulated = false;
+
+	// Raw data GUID taken from UAnimSequence when initially populating - this allows for retaining compressed data state initially
+	UPROPERTY(VisibleAnywhere, Category=AnimSequencer)
+	FGuid CachedRawDataGUID;
+
+	// Scope lock to prevent contention around ControlRig->Evaluation()
+	mutable FCriticalSection EvaluationLock;
+
+	friend class UAnimSequencerController;
+};
diff --git a/Engine/Plugins/Animation/AnimationModifierLibrary/Source/AnimationModifierLibrary/Private/MotionExtractorModifier.cpp b/Engine/Plugins/Animation/AnimationModifierLibrary/Source/AnimationModifierLibrary/Private/MotionExtractorModifier.cpp
--- a/Engine/Plugins/Animation/AnimationModifierLibrary/Source/AnimationModifierLibrary/Private/MotionExtractorModifier.cpp
+++ b/Engine/Plugins/Animation/AnimationModifierLibrary/Source/AnimationModifierLibrary/Private/MotionExtractorModifier.cpp
@@ -256,7 +256,7 @@ FTransform UMotionExtractorModifier::ExtractBoneTransform(UAnimSequence* Animati
 	FBlendedCurve Curve;
 	Curve.InitFrom(BoneContainer);
 
-	FAnimExtractContext Context(Time, false);
+	FAnimExtractContext Context(static_cast<double>(Time), false);
 	UE::Anim::FStackAttributeContainer Attributes;
 	FAnimationPoseData AnimationPoseData(Pose, Curve, Attributes);
 
diff --git a/Engine/Plugins/Animation/ControlRig/Source/ControlRig/Private/Sequencer/MovieSceneControlRigParameterSection.cpp b/Engine/Plugins/Animation/ControlRig/Source/ControlRig/Private/Sequencer/MovieSceneControlRigParameterSection.cpp
--- a/Engine/Plugins/Animation/ControlRig/Source/ControlRig/Private/Sequencer/MovieSceneControlRigParameterSection.cpp
+++ b/Engine/Plugins/Animation/ControlRig/Source/ControlRig/Private/Sequencer/MovieSceneControlRigParameterSection.cpp
@@ -2519,11 +2519,11 @@ bool UMovieSceneControlRigParameterSection::LoadAnimSequenceIntoThisSection(UAni
 	}
 	ControlRig->Modify();
 
-	const int32 NumberOfFrames = FrameRate.AsFrameTime(Length).CeilToFrame().Value + 1;
+	const int32 NumberOfKeys = AnimSequence->GetDataModel()->GetNumberOfKeys();
 	FFrameNumber FrameRateInFrameNumber = TickResolution.AsFrameNumber(FrameRate.AsInterval());
 	int32 ExtraProgress = bKeyReduce ? FloatChannels.Num() : 0;
 	
-	FScopedSlowTask Progress(NumberOfFrames + ExtraProgress, LOCTEXT("BakingToControlRig_SlowTask", "Baking To Control Rig..."));	
+	FScopedSlowTask Progress(NumberOfKeys + ExtraProgress, LOCTEXT("BakingToControlRig_SlowTask", "Baking To Control Rig..."));	
 	Progress.MakeDialog(true);
 
 	//Make sure we are reset and run setup event  before evaluating
@@ -2545,7 +2545,7 @@ bool UMovieSceneControlRigParameterSection::LoadAnimSequenceIntoThisSection(UAni
 	SourceCurves.ResetValues();
 	ControlRig->Execute(EControlRigState::Update, TEXT("Setup"));
 	*/
-	const UAnimDataModel* DataModel = AnimSequence->GetDataModel();
+	const IAnimationDataModel* DataModel = AnimSequence->GetDataModel();
 	const FAnimationCurveData& CurveData = DataModel->GetCurveData();
 	const TArray<FBoneAnimationTrack>& BoneAnimationTracks = DataModel->GetBoneAnimationTracks();
 
@@ -2571,10 +2571,10 @@ bool UMovieSceneControlRigParameterSection::LoadAnimSequenceIntoThisSection(UAni
 	ControlRig->RequestSetup();
 	ControlRig->Evaluate_AnyThread();
 
-	for (int32 Index = 0; Index < NumberOfFrames; ++Index)
+	for (int32 Index = 0; Index < NumberOfKeys; ++Index)
 	{
 		const float SequenceSecond = AnimSequence->GetTimeAtFrame(Index);
-		FFrameNumber FrameNumber = StartFrame + (FrameRateInFrameNumber * Index);
+		const FFrameNumber FrameNumber = StartFrame + (FMath::Max(FrameRateInFrameNumber.Value, 1) * Index);
 
 		ControlRig->GetHierarchy()->ResetPoseToInitial();
 		ControlRig->GetHierarchy()->ResetCurveValues();
diff --git a/Engine/Plugins/Animation/IKRig/Source/IKRigEditor/Private/RetargetEditor/IKRetargetBatchOperation.cpp b/Engine/Plugins/Animation/IKRig/Source/IKRigEditor/Private/RetargetEditor/IKRetargetBatchOperation.cpp
--- a/Engine/Plugins/Animation/IKRig/Source/IKRigEditor/Private/RetargetEditor/IKRetargetBatchOperation.cpp
+++ b/Engine/Plugins/Animation/IKRig/Source/IKRigEditor/Private/RetargetEditor/IKRetargetBatchOperation.cpp
@@ -186,6 +186,7 @@ void FIKRetargetBatchOperation::RetargetAssets(
 			// set the retarget source to the target skeletal mesh
 			AnimSequenceToRetarget->RetargetSource = NAME_None;
 			AnimSequenceToRetarget->RetargetSourceAsset = Context.TargetMesh;
+			Controller.UpdateWithSkeleton(NewSkeleton);
 			// done editing sequence data, close bracket
 			Controller.CloseBracket(ShouldTransactAnimEdits);
 		}
diff --git a/Engine/Plugins/Experimental/Animation/ContextualAnimation/Source/ContextualAnimation/Private/ContextualAnimTypes.cpp b/Engine/Plugins/Experimental/Animation/ContextualAnimation/Source/ContextualAnimation/Private/ContextualAnimTypes.cpp
--- a/Engine/Plugins/Experimental/Animation/ContextualAnimation/Source/ContextualAnimation/Private/ContextualAnimTypes.cpp
+++ b/Engine/Plugins/Experimental/Animation/ContextualAnimation/Source/ContextualAnimation/Private/ContextualAnimTypes.cpp
@@ -38,7 +38,7 @@ FTransform FContextualAnimAlignmentTrackContainer::ExtractTransformAtTime(int32
 		const FRawAnimSequenceTrack& Track = Tracks.AnimationTracks[TrackIndex];
 		const int32 TotalFrames = Track.PosKeys.Num();
 		const float TrackLength = (TotalFrames - 1) * SampleInterval;
-		FAnimationUtils::ExtractTransformFromTrack(Time, TotalFrames, TrackLength, Track, EAnimInterpolationType::Linear, AlignmentTransform);
+		FAnimationUtils::ExtractTransformFromTrack(Track, Time, TotalFrames, TrackLength, EAnimInterpolationType::Linear, AlignmentTransform);
 	}
 
 	return AlignmentTransform;
diff --git a/Engine/Plugins/Experimental/Animation/ContextualAnimation/Source/ContextualAnimation/Private/ContextualAnimUtilities.cpp b/Engine/Plugins/Experimental/Animation/ContextualAnimation/Source/ContextualAnimation/Private/ContextualAnimUtilities.cpp
--- a/Engine/Plugins/Experimental/Animation/ContextualAnimation/Source/ContextualAnimation/Private/ContextualAnimUtilities.cpp
+++ b/Engine/Plugins/Experimental/Animation/ContextualAnimation/Source/ContextualAnimation/Private/ContextualAnimUtilities.cpp
@@ -28,7 +28,7 @@ void UContextualAnimUtilities::ExtractLocalSpacePose(const UAnimSequenceBase* An
 	FBlendedCurve Curve;
 	Curve.InitFrom(BoneContainer);
 
-	FAnimExtractContext Context(Time, bExtractRootMotion);
+	FAnimExtractContext Context(static_cast<double>(Time), bExtractRootMotion);
 
 	UE::Anim::FStackAttributeContainer Attributes;
 	FAnimationPoseData AnimationPoseData(OutPose, Curve, Attributes);
diff --git a/Engine/Plugins/Experimental/Animation/MotionWarping/Source/MotionWarping/Private/MotionWarpingComponent.cpp b/Engine/Plugins/Experimental/Animation/MotionWarping/Source/MotionWarping/Private/MotionWarpingComponent.cpp
--- a/Engine/Plugins/Experimental/Animation/MotionWarping/Source/MotionWarping/Private/MotionWarpingComponent.cpp
+++ b/Engine/Plugins/Experimental/Animation/MotionWarping/Source/MotionWarping/Private/MotionWarpingComponent.cpp
@@ -31,7 +31,7 @@
 	FBlendedCurve Curve;
 	Curve.InitFrom(BoneContainer);
 
-	FAnimExtractContext Context(Time, bExtractRootMotion);
+	FAnimExtractContext Context(static_cast<double>(Time), bExtractRootMotion);
 
 	UE::Anim::FStackAttributeContainer Attributes;
 	FAnimationPoseData AnimationPoseData(OutPose, Curve, Attributes);
diff --git a/Engine/Plugins/Experimental/Animation/MotionWarping/Source/MotionWarping/Private/RootMotionModifier_AdjustmentBlendWarp.cpp b/Engine/Plugins/Experimental/Animation/MotionWarping/Source/MotionWarping/Private/RootMotionModifier_AdjustmentBlendWarp.cpp
--- a/Engine/Plugins/Experimental/Animation/MotionWarping/Source/MotionWarping/Private/RootMotionModifier_AdjustmentBlendWarp.cpp
+++ b/Engine/Plugins/Experimental/Animation/MotionWarping/Source/MotionWarping/Private/RootMotionModifier_AdjustmentBlendWarp.cpp
@@ -37,7 +37,7 @@
 	const float TrackLength = (EndTime - ActualStartTime);
 	const float TimePercent = (Time - ActualStartTime) / (EndTime - ActualStartTime);
 	const float RemappedTime = TimePercent * TrackLength;
-	FAnimationUtils::ExtractTransformFromTrack(RemappedTime, TotalFrames, TrackLength, Result.AnimationTracks[TrackIndex], EAnimInterpolationType::Linear, OutTransform);
+	FAnimationUtils::ExtractTransformFromTrack(Result.AnimationTracks[TrackIndex], static_cast<double>(RemappedTime), TotalFrames, TrackLength, EAnimInterpolationType::Linear, OutTransform);
 }
 
 void URootMotionModifier_AdjustmentBlendWarp::ExtractBoneTransformAtFrame(FTransform& OutTransform, int32 TrackIndex, int32 Frame) const
diff --git a/Engine/Plugins/Experimental/Animation/PoseSearch/Source/Editor/Private/PoseSearchDebugger.cpp b/Engine/Plugins/Experimental/Animation/PoseSearch/Source/Editor/Private/PoseSearchDebugger.cpp
--- a/Engine/Plugins/Experimental/Animation/PoseSearch/Source/Editor/Private/PoseSearchDebugger.cpp
+++ b/Engine/Plugins/Experimental/Animation/PoseSearch/Source/Editor/Private/PoseSearchDebugger.cpp
@@ -186,7 +186,7 @@ void UPoseSearchMeshComponent::UpdatePose(const FUpdateContext& UpdateContext)
 		
 		FDeltaTimeRecord DeltaTimeRecord;
 		DeltaTimeRecord.Set(PreviousTime, AdvancedTime - PreviousTime);
-		FAnimExtractContext ExtractionCtx(AdvancedTime, true, DeltaTimeRecord, UpdateContext.bLoop);
+		FAnimExtractContext ExtractionCtx(static_cast<double>(AdvancedTime), true, DeltaTimeRecord, UpdateContext.bLoop);
 
 		for (int32 BlendSampleIdex = 0; BlendSampleIdex < BlendSamples.Num(); BlendSampleIdex++)
 		{
diff --git a/Engine/Plugins/Experimental/Animation/PoseSearch/Source/Runtime/Private/PoseSearch.cpp b/Engine/Plugins/Experimental/Animation/PoseSearch/Source/Runtime/Private/PoseSearch.cpp
--- a/Engine/Plugins/Experimental/Animation/PoseSearch/Source/Runtime/Private/PoseSearch.cpp
+++ b/Engine/Plugins/Experimental/Animation/PoseSearch/Source/Runtime/Private/PoseSearch.cpp
@@ -3232,7 +3232,7 @@ void FBlendSpaceSampler::ProcessRootTransform()
 
 		FDeltaTimeRecord DeltaTimeRecord;
 		DeltaTimeRecord.Set(PreviousTime, CurrentTime - PreviousTime);
-		FAnimExtractContext ExtractionCtx(CurrentTime, true, DeltaTimeRecord, Input.bLoopable);
+		FAnimExtractContext ExtractionCtx(static_cast<double>(CurrentTime), true, DeltaTimeRecord, Input.bLoopable);
 
 		for (int32 BlendSampleIdex = 0; BlendSampleIdex < BlendSamples.Num(); BlendSampleIdex++)
 		{
diff --git a/Engine/Plugins/Experimental/Animation/PoseSearch/Source/Runtime/Private/PoseSearchDerivedData.cpp b/Engine/Plugins/Experimental/Animation/PoseSearch/Source/Runtime/Private/PoseSearchDerivedData.cpp
--- a/Engine/Plugins/Experimental/Animation/PoseSearch/Source/Runtime/Private/PoseSearchDerivedData.cpp
+++ b/Engine/Plugins/Experimental/Animation/PoseSearch/Source/Runtime/Private/PoseSearchDerivedData.cpp
@@ -278,7 +278,7 @@ namespace UE::PoseSearch
 		if (Sequence)
 		{
 			InOutWriter.Update(MakeMemoryView(Sequence->GetName()));
-			InOutWriter.Update(MakeMemoryView(Sequence->GetRawDataGuid().ToString()));
+			InOutWriter.Update(MakeMemoryView(Sequence->GetDataModel()->GenerateGuid().ToString()));
 		}
 	}
 
diff --git a/Engine/Plugins/Experimental/Animation/PoseSearch/Source/Runtime/Private/PoseSearchFeatureChannels.cpp b/Engine/Plugins/Experimental/Animation/PoseSearch/Source/Runtime/Private/PoseSearchFeatureChannels.cpp
--- a/Engine/Plugins/Experimental/Animation/PoseSearch/Source/Runtime/Private/PoseSearchFeatureChannels.cpp
+++ b/Engine/Plugins/Experimental/Animation/PoseSearch/Source/Runtime/Private/PoseSearchFeatureChannels.cpp
@@ -192,7 +192,7 @@ void UPoseSearchFeatureChannel_Pose::AddPoseFeatures(const  UE::PoseSearch::IAss
 
 			FDeltaTimeRecord DeltaTimeRecord;
 			DeltaTimeRecord.Set(PreviousTime, CurrentTime - PreviousTime);
-			FAnimExtractContext ExtractionCtx(CurrentTime, true, DeltaTimeRecord, Samples[Term].Clip->IsLoopable());
+			FAnimExtractContext ExtractionCtx(static_cast<double>(CurrentTime), true, DeltaTimeRecord, Samples[Term].Clip->IsLoopable());
 
 			Samples[Term].Clip->ExtractPose(ExtractionCtx, AnimPoseData[Term]);
 
diff --git a/Engine/Plugins/Experimental/Interchange/Runtime/Source/Import/Private/Animation/InterchangeAnimSequenceFactory.cpp b/Engine/Plugins/Experimental/Interchange/Runtime/Source/Import/Private/Animation/InterchangeAnimSequenceFactory.cpp
--- a/Engine/Plugins/Experimental/Interchange/Runtime/Source/Import/Private/Animation/InterchangeAnimSequenceFactory.cpp
+++ b/Engine/Plugins/Experimental/Interchange/Runtime/Source/Import/Private/Animation/InterchangeAnimSequenceFactory.cpp
@@ -197,7 +197,7 @@ namespace UE::Interchange::Private
 		FAnimationCurveIdentifier FloatCurveId(NewName, ERawCurveTrackTypes::RCT_Float);
 
 
-		UAnimDataModel* DataModel = TargetSequence->GetDataModel();
+		IAnimationDataModel* DataModel = TargetSequence->GetDataModel();
 		IAnimationDataController& Controller = TargetSequence->GetController();
 
 		const FFloatCurve* TargetCurve = DataModel->FindFloatCurve(FloatCurveId);
@@ -312,7 +312,11 @@ namespace UE::Interchange::Private
 
 			//This destroy all previously imported animation raw data
 			Controller.RemoveAllBoneTracks();
-			Controller.SetPlayLength(FGenericPlatformMath::Max<float>(SequenceLength, MINIMUM_ANIMATION_LENGTH));
+
+			const FFrameRate ResampleFrameRate(SampleRate, 1);
+			const FFrameNumber NumberOfFrames = ResampleFrameRate.AsFrameNumber(SequenceLength);
+			Controller.SetFrameRate(ResampleFrameRate);
+			Controller.SetNumberOfFrames(FGenericPlatformMath::Max<int32>(NumberOfFrames.Value, 1));
 
 			FTransform3f GlobalOffsetTransform;
 			{
@@ -444,7 +448,7 @@ namespace UE::Interchange::Private
 		AnimSequenceFactoryNode->GetCustomImportAttributeCurves(bImportAttributeCurves);
 		if (bImportAttributeCurves)
 		{
-			const UAnimDataModel* DataModel = AnimSequence->GetDataModel();
+			const IAnimationDataModel* DataModel = AnimSequence->GetDataModel();
 			const int32 NumFloatCurves = DataModel->GetNumberOfFloatCurves();
 			const FAnimationCurveData& CurveData = DataModel->GetCurveData();
 
@@ -893,6 +897,7 @@ UObject* UInterchangeAnimSequenceFactory::CreateAsset(const FCreateAssetParams&
 		IAnimationDataController& Controller = AnimSequence->GetController();
 		Controller.OpenBracket(NSLOCTEXT("InterchangeAnimSequenceFactory", "ImportAnimationInterchange_Bracket", "Importing Animation (Interchange)"));
 		AnimSequence->SetSkeleton(Skeleton);
+		Controller.InitializeModel();
 		AnimSequence->ImportFileFramerate = SampleRate;
 		AnimSequence->ImportResampleFramerate = SampleRate;
 		Controller.SetFrameRate(FrameRate);
diff --git a/Engine/Plugins/Experimental/MLDeformer/MLDeformerFramework/Source/MLDeformerFrameworkEditor/Private/MLDeformerEditorModel.cpp b/Engine/Plugins/Experimental/MLDeformer/MLDeformerFramework/Source/MLDeformerFrameworkEditor/Private/MLDeformerEditorModel.cpp
--- a/Engine/Plugins/Experimental/MLDeformer/MLDeformerFramework/Source/MLDeformerFrameworkEditor/Private/MLDeformerEditorModel.cpp
+++ b/Engine/Plugins/Experimental/MLDeformer/MLDeformerFramework/Source/MLDeformerFrameworkEditor/Private/MLDeformerEditorModel.cpp
@@ -869,7 +869,7 @@ namespace UE::MLDeformer
 			return;
 		}
 
-		const UAnimDataModel* DataModel = Model->AnimSequence->GetDataModel();
+		const IAnimationDataModel* DataModel = Model->AnimSequence->GetDataModel();
 		if (!DataModel)
 		{
 			UE_LOG(LogMLDeformer, Warning, TEXT("Anim sequence has no data model."));
@@ -957,7 +957,7 @@ namespace UE::MLDeformer
 			return;
 		}
 
-		const UAnimDataModel* DataModel = Model->AnimSequence->GetDataModel();
+		const IAnimationDataModel* DataModel = Model->AnimSequence->GetDataModel();
 		if (!DataModel)
 		{
 			UE_LOG(LogMLDeformer, Warning, TEXT("Anim sequence has no data model."));
diff --git a/Engine/Plugins/Experimental/SkeletalReduction/Source/Private/SkeletalMeshReductionPlugin.cpp b/Engine/Plugins/Experimental/SkeletalReduction/Source/Private/SkeletalMeshReductionPlugin.cpp
--- a/Engine/Plugins/Experimental/SkeletalReduction/Source/Private/SkeletalMeshReductionPlugin.cpp
+++ b/Engine/Plugins/Experimental/SkeletalReduction/Source/Private/SkeletalMeshReductionPlugin.cpp
@@ -1202,7 +1202,7 @@ float FQuadricSkeletalMeshReduction::SimplifyMesh( const FSkeletalMeshOptimizati
 
 	// Copy the simplified mesh back into Mesh
 
-	Simplifier.OutputMesh(Mesh.VertexBuffer, Mesh.IndexBuffer, bWeldVertexColorAttributes, NULL);
+	Simplifier.OutputMesh(Mesh.VertexBuffer, Mesh.IndexBuffer, bWeldVertexColorAttributes, nullptr);
 
 	// There might have some unused verts at the end of the vertex buffer that were generated by the possible duplicates
 
@@ -1894,7 +1894,7 @@ void FQuadricSkeletalMeshReduction::ReduceSkeletalMesh(USkeletalMesh& SkeletalMe
 	// Insert a new LOD model entry if needed.
 	if (LODIndex == SkeletalMeshResource.LODModels.Num())
 	{
-		FSkeletalMeshLODModel* ModelPtr = NULL;
+		FSkeletalMeshLODModel* ModelPtr = nullptr;
 		SkeletalMeshResource.LODModels.Add(ModelPtr);
 		bLODModelAdded = true;
 	}
@@ -2119,7 +2119,7 @@ void FQuadricSkeletalMeshReduction::ReduceSkeletalMesh(USkeletalMesh& SkeletalMe
 
 		// Get component space retarget base pose, will be equivalent of ref-pose if not edited
 		TArray<FTransform> ComponentSpaceRefPose;
-		FAnimationRuntime::FillUpComponentSpaceTransformsRetargetBasePose(&SkeletalMesh, ComponentSpaceRefPose); 
+		FAnimationRuntime::FillUpComponentSpaceTransforms(SkeletalMesh.GetRefSkeleton(), SkeletalMesh.GetRefSkeleton().GetRefBonePose(), ComponentSpaceRefPose);
 		
 		// Retrieve bone indices which will be removed
 		if (MeshBoneReductionInterface != nullptr)
@@ -2142,10 +2142,16 @@ void FQuadricSkeletalMeshReduction::ReduceSkeletalMesh(USkeletalMesh& SkeletalMe
 		Pose.SetBoneContainer(&RequiredBones);
 		Pose.ResetToRefPose();
 
+		FBlendedCurve TempCurve;
+		TempCurve.InitFrom(RequiredBones);
+		UE::Anim::FStackAttributeContainer TempAttributes;
+		
+		FAnimationPoseData AnimPose(Pose, TempCurve, TempAttributes);
+
 		// Retrieve animated pose from anim sequence (including retargeting)
 		const USkeleton* Skeleton = SkeletalMesh.GetSkeleton();
 		const FName RetargetSource = Skeleton->GetRetargetSourceForMesh(&SkeletalMesh);
-		UE::Anim::BuildPoseFromModel(BakePoseAnim->GetDataModel(), Pose, 0.f, EAnimInterpolationType::Step, RetargetSource, Skeleton->GetRefLocalPoses(RetargetSource));
+		UE::Anim::BuildPoseFromModel(BakePoseAnim->GetDataModel(), AnimPose, 0.0, EAnimInterpolationType::Step, RetargetSource, Skeleton->GetRefLocalPoses(RetargetSource));
 		
 		// Calculate component space animated pose matrices
 		TArray<FMatrix> ComponentSpaceAnimatedPose;
@@ -2246,7 +2252,7 @@ void FQuadricSkeletalMeshReduction::ReduceSkeletalMesh(USkeletalMesh& SkeletalMe
 		check(ReducedLODInfoPtr);
 		// Do any joint-welding / bone removal.
 
-		if (MeshBoneReductionInterface != NULL && MeshBoneReductionInterface->GetBoneReductionData(&SkeletalMesh, LODIndex, BonesToRemove))
+		if (MeshBoneReductionInterface != nullptr && MeshBoneReductionInterface->GetBoneReductionData(&SkeletalMesh, LODIndex, BonesToRemove))
 		{
 			// fix up chunks to remove the bones that set to be removed
 			for (int32 SectionIndex = 0; SectionIndex < NewModel->Sections.Num(); ++SectionIndex)
@@ -2289,7 +2295,7 @@ void FQuadricSkeletalMeshReduction::ReduceSkeletalMesh(USkeletalMesh& SkeletalMe
 		FSkeletalMeshLODModel::CopyStructure(NewModel, SrcModel);
 
 		// Do any joint-welding / bone removal.
-		if (MeshBoneReductionInterface != NULL && MeshBoneReductionInterface->GetBoneReductionData(&SkeletalMesh, LODIndex, BonesToRemove))
+		if (MeshBoneReductionInterface != nullptr && MeshBoneReductionInterface->GetBoneReductionData(&SkeletalMesh, LODIndex, BonesToRemove))
 		{
 			// fix up chunks to remove the bones that set to be removed
 			for (int32 SectionIndex = 0; SectionIndex < NewModel->Sections.Num(); ++SectionIndex)
diff --git a/Engine/Plugins/Importers/AlembicImporter/Source/AlembicLibrary/Private/AbcImporter.cpp b/Engine/Plugins/Importers/AlembicImporter/Source/AlembicLibrary/Private/AbcImporter.cpp
--- a/Engine/Plugins/Importers/AlembicImporter/Source/AlembicLibrary/Private/AbcImporter.cpp
+++ b/Engine/Plugins/Importers/AlembicImporter/Source/AlembicLibrary/Private/AbcImporter.cpp
@@ -575,7 +575,7 @@ TArray<UObject*> FAbcImporter::ImportAsSkeletalMesh(UObject* InParent, EObjectFl
 	if (SkeletalMesh)
 	{
 		// Touch pre edit change
-		SkeletalMesh->PreEditChange(NULL);
+		SkeletalMesh->PreEditChange(nullptr);
 
 		// Retrieve the imported resource structure and allocate a new LOD model
 		FSkeletalMeshModel* ImportedModel = SkeletalMesh->GetImportedModel();
@@ -651,12 +651,15 @@ TArray<UObject*> FAbcImporter::ImportAsSkeletalMesh(UObject* InParent, EObjectFl
 		IAnimationDataController& Controller = Sequence->GetController();
 
 		Controller.OpenBracket(LOCTEXT("ImportAsSkeletalMesh", "Importing Alembic Animation"));
+		Controller.InitializeModel();
 
-		Controller.SetPlayLength(AbcFile->GetImportLength());
-		Controller.SetFrameRate( FFrameRate(AbcFile->GetFramerate(), 1));
+		const FFrameRate FrameRate(AbcFile->GetFramerate(), 1);
+		Controller.SetFrameRate(FrameRate);	
+		const FFrameNumber FrameNumber = FrameRate.AsFrameNumber(AbcFile->GetImportLength());
+		Controller.SetNumberOfFrames(FrameNumber);
 
-		Sequence->ImportFileFramerate = AbcFile->GetFramerate();
-		Sequence->ImportResampleFramerate = 1.0f / (float)AbcFile->GetFramerate();
+		Sequence->ImportFileFramerate = FrameRate.AsDecimal();
+		Sequence->ImportResampleFramerate = FrameRate.AsInterval();
 
 		{
 #if WITH_EDITOR
@@ -1748,7 +1751,7 @@ bool FAbcImporter::BuildSkeletalMesh( FSkeletalMeshLODModel& LODModel, const FRe
 	}
 
 	// Compute the required bones for this model.
-	USkeletalMesh::CalculateRequiredBones(LODModel, RefSkeleton, NULL);
+	USkeletalMesh::CalculateRequiredBones(LODModel, RefSkeleton, nullptr);
 
 	return true;
 }
diff --git a/Engine/Plugins/Importers/USDImporter/Source/USDUtilities/Private/USDSkeletalDataConversion.cpp b/Engine/Plugins/Importers/USDImporter/Source/USDUtilities/Private/USDSkeletalDataConversion.cpp
--- a/Engine/Plugins/Importers/USDImporter/Source/USDUtilities/Private/USDSkeletalDataConversion.cpp
+++ b/Engine/Plugins/Importers/USDImporter/Source/USDUtilities/Private/USDSkeletalDataConversion.cpp
@@ -420,11 +420,11 @@ namespace SkelDataConversionImpl
 		Skeleton->AddSmartNameAndModify( USkeleton::AnimCurveMappingName, CurveName, NewName );
 
 		const bool bShouldTransact = false;
-		const UAnimDataModel* DataModel = Sequence->GetDataModel();
+		const IAnimationDataModel* DataModel = Sequence->GetDataModel();
 		IAnimationDataController& Controller = Sequence->GetController();
 
-		FAnimationCurveIdentifier CurveId( NewName, ERawCurveTrackTypes::RCT_Float );
-		const FFloatCurve* Curve = DataModel->FindFloatCurve( CurveId );
+		FAnimationCurveIdentifier CurveId(NewName, ERawCurveTrackTypes::RCT_Float);
+		const FFloatCurve* Curve = DataModel->FindFloatCurve(CurveId);
 		if ( !Curve )
 		{
 			// If curve doesn't exist, add one
@@ -1913,6 +1913,7 @@ bool UsdToUnreal::ConvertSkelAnim( const pxr::UsdSkelSkeletonQuery& InUsdSkeleto
 	// it will also create a transaction when importing into UE assets, and the level sequence assets can emit some warnings about it
 	const bool bShouldTransact = false;
 	Controller.OpenBracket( LOCTEXT( "ImportUSDAnimData_Bracket", "Importing USD Animation Data" ), bShouldTransact );
+	Controller.InitializeModel();
 	Controller.ResetModel( bShouldTransact );
 
 	// Bake the animation for each frame.
@@ -2145,8 +2146,11 @@ bool UsdToUnreal::ConvertSkelAnim( const pxr::UsdSkelSkeletonQuery& InUsdSkeleto
 	OutSkeletalAnimationAsset->ImportFileFramerate = LayerTimeCodesPerSecond;
 	OutSkeletalAnimationAsset->ImportResampleFramerate = LayerTimeCodesPerSecond;
 
-	Controller.SetPlayLength( LayerSequenceLengthSeconds, bShouldTransact );
-	Controller.SetFrameRate( FFrameRate( LayerTimeCodesPerSecond, 1 ), bShouldTransact );
+
+	const FFrameRate FrameRate(LayerTimeCodesPerSecond, 1);
+	Controller.SetFrameRate(FrameRate, bShouldTransact);
+	const FFrameNumber FrameNumber = FrameRate.AsFrameNumber(LayerSequenceLengthSeconds);
+	Controller.SetNumberOfFrames(FrameNumber, bShouldTransact);
 	Controller.NotifyPopulated(); // This call is important to get the controller to not use the sampling frequency as framerate
 	Controller.CloseBracket( bShouldTransact );
 
diff --git a/Engine/Plugins/VirtualProduction/Takes/Source/TakeTrackRecorders/Private/MovieScene3DTransformTrackRecorder.cpp b/Engine/Plugins/VirtualProduction/Takes/Source/TakeTrackRecorders/Private/MovieScene3DTransformTrackRecorder.cpp
--- a/Engine/Plugins/VirtualProduction/Takes/Source/TakeTrackRecorders/Private/MovieScene3DTransformTrackRecorder.cpp
+++ b/Engine/Plugins/VirtualProduction/Takes/Source/TakeTrackRecorders/Private/MovieScene3DTransformTrackRecorder.cpp
@@ -12,7 +12,7 @@
 #include "Components/SceneComponent.h"
 #include "Components/SkeletalMeshComponent.h"
 #include "SequenceRecorderUtils.h"
-#include "Animation/AnimData/AnimDataModel.h"
+#include "Animation/AnimData/IAnimationDataModel.h"
 
 
 DEFINE_LOG_CATEGORY(TransformSerialization);
diff --git a/Engine/Plugins/VirtualProduction/Takes/Source/TakeTrackRecorders/Private/MovieSceneAnimationTrackRecorder.cpp b/Engine/Plugins/VirtualProduction/Takes/Source/TakeTrackRecorders/Private/MovieSceneAnimationTrackRecorder.cpp
--- a/Engine/Plugins/VirtualProduction/Takes/Source/TakeTrackRecorders/Private/MovieSceneAnimationTrackRecorder.cpp
+++ b/Engine/Plugins/VirtualProduction/Takes/Source/TakeTrackRecorders/Private/MovieSceneAnimationTrackRecorder.cpp
@@ -390,12 +390,13 @@ bool UMovieSceneAnimationTrackRecorder::LoadRecordedFile(const FString& FileName
 								IAnimationDataController& Controller = AnimSequence->GetController();
 								{
 									IAnimationDataController::FScopedBracket ScopedBracket(Controller, LOCTEXT("LoadRecordedFile_Bracket", "Loading recorded animation file"));
-
+									Controller.InitializeModel();
 									Controller.ResetModel();
 
 									const float FloatDenominator = 1000.0f;
 									const float Numerator = FloatDenominator / Header.IntervalTime;
-									Controller.SetFrameRate(FFrameRate(Numerator, FloatDenominator));
+									const FFrameRate FrameRate(Numerator, FloatDenominator);
+									Controller.SetFrameRate(FrameRate);
 
 									int32 MaxNumberOfKeys = 0;
 									for (int32 TrackIndex = 0; TrackIndex < Header.AnimationTrackNames.Num(); ++TrackIndex)
@@ -422,7 +423,9 @@ bool UMovieSceneAnimationTrackRecorder::LoadRecordedFile(const FString& FileName
 										Controller.SetBoneTrackKeys(Header.AnimationTrackNames[TrackIndex], PosKeys, RotKeys, ScaleKeys);
 									}
 
-									Controller.SetPlayLength((MaxNumberOfKeys > 1) ? (MaxNumberOfKeys - 1) * Header.IntervalTime : MINIMUM_ANIMATION_LENGTH);
+									const FFrameNumber FrameNumber = FrameRate.AsFrameNumber((MaxNumberOfKeys > 1) ? (MaxNumberOfKeys - 1) * Header.IntervalTime : FrameRate.AsInterval());
+									Controller.SetNumberOfFrames(FrameNumber);
+
 									Controller.NotifyPopulated();
 								}
 
@@ -436,7 +439,7 @@ bool UMovieSceneAnimationTrackRecorder::LoadRecordedFile(const FString& FileName
 								FSavePackageArgs SaveArgs;
 								SaveArgs.TopLevelFlags = RF_Standalone;
 								SaveArgs.SaveFlags = SAVE_NoError;
-								UPackage::SavePackage(Package, NULL, *PackageFileName, SaveArgs);
+								UPackage::SavePackage(Package, nullptr, *PackageFileName, SaveArgs);
 
 
 								FFrameNumber SequenceLength = (AnimSequence->GetPlayLength() * TickResolution).FloorToFrame();
