Fix-up Engine Localization warnings (duplicate entries) in AnimSequencerController.cpp

diff --git a/Engine/Plugins/Animation/AnimationData/Source/AnimationData/Private/AnimSequencerController.cpp b/Engine/Plugins/Animation/AnimationData/Source/AnimationData/Private/AnimSequencerController.cpp
--- a/Engine/Plugins/Animation/AnimationData/Source/AnimationData/Private/AnimSequencerController.cpp
+++ b/Engine/Plugins/Animation/AnimationData/Source/AnimationData/Private/AnimSequencerController.cpp
@@ -892,7 +892,7 @@ bool UAnimSequencerController::SetCurveFlags(const FAnimationCurveIdentifier& Cu
 
 	if (Curve)
 	{
-		FTransaction Transaction = ConditionalTransaction(LOCTEXT("SetCurveFlag", "Setting Raw Curve Flags"), bShouldTransact);
+		FTransaction Transaction = ConditionalTransaction(LOCTEXT("SetCurveFlags", "Setting Raw Curve Flags"), bShouldTransact);
 
 		const int32 CurrentFlags = Curve->GetCurveTypeFlags();
 
@@ -1285,7 +1285,7 @@ bool UAnimSequencerController::SetCurveKey(const FAnimationCurveIdentifier& Curv
 	}
 	else
 	{
-		ReportErrorf(LOCTEXT("RichCurveNotFoundError", "Unable to find rich curve: rich curve name {0}"), FText::FromName(CurveId.InternalName.DisplayName));
+		ReportErrorf(LOCTEXT("RichCurveNotFoundError", "Unable to find rich curve: curve name {0}"), FText::FromName(CurveId.InternalName.DisplayName));
 	}
 
 	return false;
