Fix crash opening animation asset

diff --git a/Engine/Plugins/Animation/AnimationData/Source/AnimationData/Private/AnimSequencerController.cpp b/Engine/Plugins/Animation/AnimationData/Source/AnimationData/Private/AnimSequencerController.cpp
--- a/Engine/Plugins/Animation/AnimationData/Source/AnimationData/Private/AnimSequencerController.cpp
+++ b/Engine/Plugins/Animation/AnimationData/Source/AnimationData/Private/AnimSequencerController.cpp
@@ -8,7 +8,7 @@
 
 #include "Animation/AnimData/IAnimationDataModel.h"
 #include "Animation/AnimData/CurveIdentifier.h"
-
+#include "Animation/AnimCurveTypes.h"
 #include "Animation/AnimSequence.h"
 
 #include "Algo/Transform.h"
@@ -409,7 +409,11 @@ void UAnimSequencerController::FindOrAddCurveNamesOnSkeleton(USkeleton* Skeleton
 								IdentifierPair.Key.InternalName = NewSmartName;
 
 								// Update float curve name itself as well 
-								Model->FindMutableFloatCurveById(FAnimationCurveIdentifier(CurrentSmartName, ERawCurveTrackTypes::RCT_Float))->Name = NewSmartName;					
+								FFloatCurve* FloatCurve = Model->FindMutableFloatCurveById(FAnimationCurveIdentifier(CurrentSmartName, ERawCurveTrackTypes::RCT_Float));
+								if (FloatCurve)
+								{
+									FloatCurve->Name = NewSmartName;
+								}
 							}
 						}
 					}
