Manual integration/reimplementation of CL 22091342:
Preinfintiy and Postinfinity settings within the Curve Editor of an Animation can not be changed from Constant
#feature added API for setting curve attributes (pre/post extrapolation)
#misc added tests for new API

diff --git a/Engine/Plugins/Animation/AnimationData/Source/AnimationData/Private/AnimSequencerController.cpp b/Engine/Plugins/Animation/AnimationData/Source/AnimationData/Private/AnimSequencerController.cpp
--- a/Engine/Plugins/Animation/AnimationData/Source/AnimationData/Private/AnimSequencerController.cpp
+++ b/Engine/Plugins/Animation/AnimationData/Source/AnimationData/Private/AnimSequencerController.cpp
@@ -1387,6 +1387,73 @@ bool UAnimSequencerController::SetCurveKeys(const FAnimationCurveIdentifier& Cur
 	return false;
 }
 
+bool UAnimSequencerController::SetCurveAttributes(const FAnimationCurveIdentifier& CurveId, const FCurveAttributes& Attributes, bool bShouldTransact)
+{
+	ValidateModel();
+
+	const FRichCurve* RichCurve = Model->GetMutableRichCurve(CurveId);
+	if (RichCurve)
+	{
+		FTransaction Transaction = ConditionalTransaction(LOCTEXT("SettingNamedCurveAttributes", "Setting Curve Attributes"), bShouldTransact);
+
+		FCurveAttributes CurrentAttributes;
+		CurrentAttributes.SetPreExtrapolation(RichCurve->PreInfinityExtrap);
+		CurrentAttributes.SetPostExtrapolation(RichCurve->PostInfinityExtrap);
+
+		if (UMovieSceneControlRigParameterSection* Section = Model->GetFKControlRigSection())
+		{
+			UControlRig* ControlRig = Section->GetControlRig();
+
+			if (UFKControlRig* FKRig = Cast<UFKControlRig>(ControlRig))
+			{
+				const URigHierarchy* Hierarchy = FKRig->GetHierarchy();
+				if (Hierarchy ||  IgnoreSkeletonValidation())
+				{
+					const FRigElementKey CurveKey(CurveId.InternalName.DisplayName, ERigElementType::Curve);
+					const FRigElementKey CurveControlKey(UFKControlRig::GetControlName(CurveId.InternalName.DisplayName, ERigElementType::Curve), ERigElementType::Control);
+
+					const bool bContainsCurve = Hierarchy ? Hierarchy->Contains(CurveKey) : IgnoreSkeletonValidation();
+					const bool bContainsCurveControl = Hierarchy ? Hierarchy->Contains(CurveControlKey) : IgnoreSkeletonValidation();
+
+					if(bContainsCurve && bContainsCurveControl)
+					{
+						const bool bHasCurveControlChannel = Section->HasScalarParameter(CurveControlKey.Name);
+						if (!bHasCurveControlChannel)
+						{
+							Section->AddScalarParameter(CurveControlKey.Name, TOptional<float>(), Model->bPopulated);
+						}
+										
+						FScalarParameterNameAndCurve* ParameterCurvePair = Section->GetScalarParameterNamesAndCurves().FindByPredicate([CurveControlKey](const FScalarParameterNameAndCurve& Parameter)
+						{
+							return Parameter.ParameterName == CurveControlKey.Name;
+						});
+						check(ParameterCurvePair);
+						
+						ConditionalAction<UE::Anim::FSetRichCurveAttributesAction>(bShouldTransact, CurveId, CurrentAttributes);	
+						if (Attributes.HasPreExtrapolation())
+						{							
+							ParameterCurvePair->ParameterCurve.PreInfinityExtrap = Attributes.GetPreExtrapolation();
+						}
+
+						if (Attributes.HasPostExtrapolation())
+						{
+							ParameterCurvePair->ParameterCurve.PostInfinityExtrap = Attributes.GetPostExtrapolation();							
+						}
+					
+						FCurveChangedPayload Payload;
+						Payload.Identifier = CurveId;
+						Model->GetNotifier().Notify(EAnimDataModelNotifyType::CurveChanged, Payload);
+
+						return true;						
+					}
+				}
+			}
+		}
+	}
+
+	return false;
+}
+
 void UAnimSequencerController::NotifyPopulated()
 {
 	QUICK_SCOPE_CYCLE_COUNTER(STAT_NotifyPopulated);
@@ -2317,6 +2384,10 @@ void UAnimSequencerController::PopulateWithExistingModel(TScriptInterface<IAnima
 					if (CurveId.IsValid() || IgnoreSkeletonValidation())
 					{
 						AddCurve(CurveId, Curve.GetCurveTypeFlags(), false);
+						FCurveAttributes Attributes;
+						Attributes.SetPreExtrapolation(Curve.FloatCurve.PreInfinityExtrap);
+						Attributes.SetPostExtrapolation(Curve.FloatCurve.PostInfinityExtrap);						
+						SetCurveAttributes(CurveId, Attributes, false);
 						SetCurveKeys(CurveId, Curve.FloatCurve.GetConstRefOfKeys(), false);
 					}
 				}
diff --git a/Engine/Plugins/Animation/AnimationData/Source/AnimationData/Private/AnimSequencerController.h b/Engine/Plugins/Animation/AnimationData/Source/AnimationData/Private/AnimSequencerController.h
--- a/Engine/Plugins/Animation/AnimationData/Source/AnimationData/Private/AnimSequencerController.h
+++ b/Engine/Plugins/Animation/AnimationData/Source/AnimationData/Private/AnimSequencerController.h
@@ -56,6 +56,7 @@ public:
 	virtual bool SetCurveKey(const FAnimationCurveIdentifier& CurveId, const FRichCurveKey& Key, bool bShouldTransact = true) override;	
 	virtual bool RemoveCurveKey(const FAnimationCurveIdentifier& CurveId, float Time, bool bShouldTransact = true) override;
 	virtual bool SetCurveKeys(const FAnimationCurveIdentifier& CurveId, const TArray<FRichCurveKey>& CurveKeys, bool bShouldTransact = true) override;
+	virtual bool SetCurveAttributes(const FAnimationCurveIdentifier& CurveId, const FCurveAttributes& Attributes, bool bShouldTransact = true) override;
 	virtual void UpdateCurveNamesFromSkeleton(const USkeleton* Skeleton, ERawCurveTrackTypes SupportedCurveType, bool bShouldTransact = true) override;
 	virtual void FindOrAddCurveNamesOnSkeleton(USkeleton* Skeleton, ERawCurveTrackTypes SupportedCurveType, bool bShouldTransact = true) override;
 	virtual bool RemoveBoneTracksMissingFromSkeleton(const USkeleton* Skeleton, bool bShouldTransact = true) override;
diff --git a/Engine/Plugins/Animation/AnimationData/Source/AnimationData/Private/AnimSequencerHelpers.cpp b/Engine/Plugins/Animation/AnimationData/Source/AnimationData/Private/AnimSequencerHelpers.cpp
--- a/Engine/Plugins/Animation/AnimationData/Source/AnimationData/Private/AnimSequencerHelpers.cpp
+++ b/Engine/Plugins/Animation/AnimationData/Source/AnimationData/Private/AnimSequencerHelpers.cpp
@@ -16,6 +16,9 @@ void AnimSequencerHelpers::ConvertFloatChannelToRichCurve(const FMovieSceneFloat
 	const int32 NumKeys = Channel.GetNumKeys();
 	OutCurve.Keys.SetNum(NumKeys);
 
+	OutCurve.PreInfinityExtrap = Channel.PreInfinityExtrap;
+	OutCurve.PostInfinityExtrap = Channel.PostInfinityExtrap;
+
 	for (int32 KeyIndex = 0; KeyIndex < NumKeys; ++KeyIndex)
 	{
 		// Convert from channel to requested frame-rate
diff --git a/Engine/Source/Developer/AnimationDataController/AnimationDataController.build.cs b/Engine/Source/Developer/AnimationDataController/AnimationDataController.build.cs
--- a/Engine/Source/Developer/AnimationDataController/AnimationDataController.build.cs
+++ b/Engine/Source/Developer/AnimationDataController/AnimationDataController.build.cs
@@ -27,7 +27,8 @@ namespace UnrealBuildTool.Rules
 				PrivateDependencyModuleNames.AddRange(
 					new string[]
 					{													
-						"UnrealEd"
+						"UnrealEd",
+						"CurveEditor"
 					}
 				);
 			}
diff --git a/Engine/Source/Developer/AnimationDataController/Private/AnimDataController.cpp b/Engine/Source/Developer/AnimationDataController/Private/AnimDataController.cpp
--- a/Engine/Source/Developer/AnimationDataController/Private/AnimDataController.cpp
+++ b/Engine/Source/Developer/AnimationDataController/Private/AnimDataController.cpp
@@ -1215,6 +1215,41 @@ bool UAnimDataController::SetCurveKeys(const FAnimationCurveIdentifier& CurveId,
 	return false;
 }
 
+
+bool UAnimDataController::SetCurveAttributes(const FAnimationCurveIdentifier& CurveId, const FCurveAttributes& Attributes, bool bShouldTransact)
+{
+	ValidateModel();
+
+	FRichCurve* RichCurve = Model->GetMutableRichCurve(CurveId);
+	if (RichCurve)
+	{
+		FTransaction Transaction = ConditionalTransaction(LOCTEXT("SettingNamedCurveAttributes", "Setting Curve Attributes"), bShouldTransact);
+
+		FCurveAttributes CurrentAttributes;
+		CurrentAttributes.SetPreExtrapolation(RichCurve->PreInfinityExtrap);
+		CurrentAttributes.SetPostExtrapolation(RichCurve->PostInfinityExtrap);		
+		ConditionalAction<UE::Anim::FSetRichCurveAttributesAction>(bShouldTransact, CurveId, CurrentAttributes);
+
+		if(Attributes.HasPreExtrapolation())
+		{
+			RichCurve->PreInfinityExtrap = Attributes.GetPreExtrapolation();
+		}
+
+		if(Attributes.HasPostExtrapolation())
+		{
+			RichCurve->PostInfinityExtrap = Attributes.GetPostExtrapolation();
+		}		
+
+		FCurveChangedPayload Payload;
+		Payload.Identifier = CurveId;
+		Model->GetNotifier().Notify(EAnimDataModelNotifyType::CurveChanged, Payload);
+
+		return true;
+	}
+
+	return false;
+}
+
 void UAnimDataController::NotifyPopulated()
 {
 	ValidateModel();
diff --git a/Engine/Source/Developer/AnimationDataController/Private/AnimDataControllerActions.cpp b/Engine/Source/Developer/AnimationDataController/Private/AnimDataControllerActions.cpp
--- a/Engine/Source/Developer/AnimationDataController/Private/AnimDataControllerActions.cpp
+++ b/Engine/Source/Developer/AnimationDataController/Private/AnimDataControllerActions.cpp
@@ -388,6 +388,26 @@ FString FSetRichCurveKeysAction::ToStringInternal() const
 	return FText::Format(LOCTEXT("SetNamedRichCurveKeysAction_Description", "Replacing keys for {0} curve '{1}'."), FText::FromString(CurveId.CurveType == ERawCurveTrackTypes::RCT_Float ? FloatLabel : TransformLabel), FText::FromName(CurveId.InternalName.DisplayName)).ToString();
 }
 
+TUniquePtr<FChange> FSetRichCurveAttributesAction::ExecuteInternal(IAnimationDataModel* Model, IAnimationDataController* Controller)
+{
+	const FRichCurve& RichCurve = Model->GetRichCurve(CurveId);
+
+	FCurveAttributes CurrentAttributes;
+	CurrentAttributes.SetPreExtrapolation(RichCurve.PreInfinityExtrap);
+	CurrentAttributes.SetPostExtrapolation(RichCurve.PostInfinityExtrap);
+
+	TUniquePtr<FChange> InverseAction = MakeUnique<FSetRichCurveAttributesAction>(CurveId, CurrentAttributes);	
+	Controller->SetCurveAttributes(CurveId, Attributes, false);
+	
+	return InverseAction;
+	
+}
+
+FString FSetRichCurveAttributesAction::ToStringInternal() const
+{
+	return FText::Format(LOCTEXT("SetCurveAttributesAction_Description", "Setting curve attributes '{0}'."), FText::FromName(CurveId.InternalName.DisplayName)).ToString();
+}
+
 TUniquePtr<FChange> FSetCurveColorAction::ExecuteInternal(IAnimationDataModel* Model, IAnimationDataController* Controller)
 {
 	const FAnimCurveBase& AnimationCurve = Model->GetCurve(CurveId);
diff --git a/Engine/Source/Developer/AnimationDataController/Public/AnimDataController.h b/Engine/Source/Developer/AnimationDataController/Public/AnimDataController.h
--- a/Engine/Source/Developer/AnimationDataController/Public/AnimDataController.h
+++ b/Engine/Source/Developer/AnimationDataController/Public/AnimDataController.h
@@ -64,6 +64,7 @@ public:
 	virtual bool SetCurveKey(const FAnimationCurveIdentifier& CurveId, const FRichCurveKey& Key, bool bShouldTransact = true) override;	
 	virtual bool RemoveCurveKey(const FAnimationCurveIdentifier& CurveId, float Time, bool bShouldTransact = true) override;
 	virtual bool SetCurveKeys(const FAnimationCurveIdentifier& CurveId, const TArray<FRichCurveKey>& CurveKeys, bool bShouldTransact = true) override;
+	virtual bool SetCurveAttributes(const FAnimationCurveIdentifier& CurveId, const FCurveAttributes& Attributes, bool bShouldTransact = true) override;
 	virtual void UpdateCurveNamesFromSkeleton(const USkeleton* Skeleton, ERawCurveTrackTypes SupportedCurveType, bool bShouldTransact = true) override;
 	virtual void FindOrAddCurveNamesOnSkeleton(USkeleton* Skeleton, ERawCurveTrackTypes SupportedCurveType, bool bShouldTransact = true) override;
 	virtual bool RemoveBoneTracksMissingFromSkeleton(const USkeleton* Skeleton, bool bShouldTransact = true) override;
diff --git a/Engine/Source/Developer/AnimationDataController/Public/AnimDataControllerActions.h b/Engine/Source/Developer/AnimationDataController/Public/AnimDataControllerActions.h
--- a/Engine/Source/Developer/AnimationDataController/Public/AnimDataControllerActions.h
+++ b/Engine/Source/Developer/AnimationDataController/Public/AnimDataControllerActions.h
@@ -2,6 +2,7 @@
 
 #pragma once
 
+#include "CurveDataAbstraction.h"
 #include "Misc/Change.h"
 #include "Misc/CoreMiscDefines.h"
 
@@ -314,6 +315,21 @@ protected:
 	TArray<FRichCurveKey> Keys;
 };
 
+class ANIMATIONDATACONTROLLER_API FSetRichCurveAttributesAction : public FAnimDataBaseAction
+{
+public:
+	explicit FSetRichCurveAttributesAction(const FAnimationCurveIdentifier& InCurveId, const FCurveAttributes& InAttributes) : CurveId(InCurveId), Attributes(InAttributes) {}
+	virtual ~FSetRichCurveAttributesAction() {}
+protected:
+	FSetRichCurveAttributesAction() {}
+	virtual TUniquePtr<FChange> ExecuteInternal(IAnimationDataModel* Model, IAnimationDataController* Controller) override;
+	virtual FString ToStringInternal() const override;
+
+protected:
+	FAnimationCurveIdentifier CurveId;
+	FCurveAttributes Attributes;
+};
+
 class ANIMATIONDATACONTROLLER_API FSetCurveColorAction : public FAnimDataBaseAction
 {
 public:
diff --git a/Engine/Source/Editor/Persona/Private/SAnimSequenceCurveEditor.cpp b/Engine/Source/Editor/Persona/Private/SAnimSequenceCurveEditor.cpp
--- a/Engine/Source/Editor/Persona/Private/SAnimSequenceCurveEditor.cpp
+++ b/Engine/Source/Editor/Persona/Private/SAnimSequenceCurveEditor.cpp
@@ -103,6 +103,13 @@ void FRichCurveEditorModelNamed::SetKeyAttributes(TArrayView<const FKeyHandle> I
 	}
 }
 
+void FRichCurveEditorModelNamed::SetCurveAttributes(const FCurveAttributes& InCurveAttributes)
+{
+	Modify();		
+	IAnimationDataController& Controller = AnimSequence->GetController();
+	Controller.SetCurveAttributes(CurveId, InCurveAttributes);
+}
+
 void FRichCurveEditorModelNamed::CurveHasChanged()
 {	
 	IAnimationDataController& Controller = AnimSequence->GetController();
diff --git a/Engine/Source/Editor/Persona/Private/SAnimSequenceCurveEditor.h b/Engine/Source/Editor/Persona/Private/SAnimSequenceCurveEditor.h
--- a/Engine/Source/Editor/Persona/Private/SAnimSequenceCurveEditor.h
+++ b/Engine/Source/Editor/Persona/Private/SAnimSequenceCurveEditor.h
@@ -33,6 +33,7 @@ public:
 
 	virtual void SetKeyPositions(TArrayView<const FKeyHandle> InKeys, TArrayView<const FKeyPosition> InKeyPositions, EPropertyChangeType::Type ChangeType) override;
 	virtual void SetKeyAttributes(TArrayView<const FKeyHandle> InKeys, TArrayView<const FKeyAttributes> InAttributes, EPropertyChangeType::Type ChangeType = EPropertyChangeType::Unspecified) override;
+	virtual void SetCurveAttributes(const FCurveAttributes& InCurveAttributes) override;
 
 	void CurveHasChanged();
 	void OnModelHasChanged(const EAnimDataModelNotifyType& NotifyType, IAnimationDataModel* Model, const FAnimDataModelNotifPayload& Payload);
diff --git a/Engine/Source/Runtime/Engine/Classes/Animation/AnimData/IAnimationDataController.h b/Engine/Source/Runtime/Engine/Classes/Animation/AnimData/IAnimationDataController.h
--- a/Engine/Source/Runtime/Engine/Classes/Animation/AnimData/IAnimationDataController.h
+++ b/Engine/Source/Runtime/Engine/Classes/Animation/AnimData/IAnimationDataController.h
@@ -24,6 +24,7 @@ class FText;
 class UAssetUserData;
 class IAnimationDataModel;
 
+struct FCurveAttributes;
 struct FAnimationCurveIdentifier;
 struct FAnimationAttributeIdentifier;
 
@@ -465,6 +466,17 @@ public:
 	UFUNCTION(BlueprintCallable, Category = CurveData)
 	virtual bool SetCurveKeys(const FAnimationCurveIdentifier& CurveId, const TArray<FRichCurveKey>& CurveKeys, bool bShouldTransact = true) = 0;
 
+	/**
+	* Changes the attributes of the curve with provided identifier. Broadcasts a EAnimDataModelNotifyType::CurveChanged notify if successful.
+	*
+	* @param	CurveId				Identifier of the curve to change the color for
+	* @param	Attributes			Attribute values to be applied
+	* @param	bShouldTransact		Whether or not any undo-redo changes should be generated
+	*
+	* @return	Whether or not the curve attributes were set successfully
+	*/
+	virtual bool SetCurveAttributes(const FAnimationCurveIdentifier& CurveId, const FCurveAttributes& Attributes, bool bShouldTransact = true) = 0;
+
 	/**
 	* Updates the display name values for any stored curve, with the names being retrieved from the provided skeleton. Broadcasts a EAnimDataModelNotifyType::CurveRenamed for each to-be-updated curve name, wrapped within BracketOpened/BracketClosed notifies.
 	*
