Crash in UAnimSequence::GetBonePose after adding curve with UAnimationBlueprintLibrary::AddCurve with Anim Modifier
#fix remove scope bracket around application of ANimationModifier, moving responsibility to the user

Curve keys added with Anim Modifiers have incorrect time
#fix ensure that when adding a scalar curve, its tick rate is matched to the AnimModel's, and convert between tick-rate and model frame rate when setting keys

diff --git a/Engine/Plugins/Animation/AnimationData/Source/AnimationData/Private/AnimSequencerController.cpp b/Engine/Plugins/Animation/AnimationData/Source/AnimationData/Private/AnimSequencerController.cpp
--- a/Engine/Plugins/Animation/AnimationData/Source/AnimationData/Private/AnimSequencerController.cpp
+++ b/Engine/Plugins/Animation/AnimationData/Source/AnimationData/Private/AnimSequencerController.cpp
@@ -1428,6 +1428,11 @@ bool UAnimSequencerController::SetCurveAttributes(const FAnimationCurveIdentifie
 							return Parameter.ParameterName == CurveControlKey.Name;
 						});
 						check(ParameterCurvePair);
+
+						if(!bHasCurveControlChannel)
+						{
+							ParameterCurvePair->ParameterCurve.SetTickResolution(Model->GetFrameRate());
+						}
 						
 						ConditionalAction<UE::Anim::FSetRichCurveAttributesAction>(bShouldTransact, CurveId, CurrentAttributes);	
 						if (Attributes.HasPreExtrapolation())
@@ -3266,6 +3271,11 @@ bool UAnimSequencerController::SetCurveControlKeys(const FName& CurveName, const
 						return Parameter.ParameterName == CurveControlKey.Name;
 					});
 					check(ParameterCurvePair);
+
+					if(!bHasCurveControlChannel)
+					{
+						ParameterCurvePair->ParameterCurve.SetTickResolution(Model->GetFrameRate());
+					}
 					
 					/*
 					 * Scenarios for curve keys conversion
@@ -3335,7 +3345,8 @@ bool UAnimSequencerController::SetCurveControlKey(const FName& CurveName, const
 				
 				if (bContainsCurve && bContainsCurveControl)
 				{
-					if(!Section->HasScalarParameter(CurveControlKey.Name))
+					const bool bHasCurveControlChannel = Section->HasScalarParameter(CurveControlKey.Name);
+					if(!bHasCurveControlChannel)
 					{
 						Section->AddScalarParameter(CurveControlKey.Name, TOptional<float>(), true);
 					}
@@ -3346,11 +3357,17 @@ bool UAnimSequencerController::SetCurveControlKey(const FName& CurveName, const
 					});
 					check(ParameterCurvePair);
 					
+					if(!bHasCurveControlChannel)
+					{
+						ParameterCurvePair->ParameterCurve.SetTickResolution(Model->GetFrameRate());
+					}
+					
 					FMovieSceneFloatValue MovieSceneValue;
 					AnimSequencerHelpers::ConvertRichCurveKeyToFloatValue(Key, MovieSceneValue );
 					const FFrameNumber FrameNumber = Model->MovieScene->GetTickResolution().AsFrameNumber(Key.Time);
+					const FFrameNumber CurveFrameNumber = ConvertFrameTime(FrameNumber, Model->MovieScene->GetTickResolution(), ParameterCurvePair->ParameterCurve.GetTickResolution()).FrameNumber;
 							
-					const FKeyHandle KeyHandle = ParameterCurvePair->ParameterCurve.GetData().UpdateOrAddKey(FrameNumber, MovieSceneValue);
+					const FKeyHandle KeyHandle = ParameterCurvePair->ParameterCurve.GetData().UpdateOrAddKey(CurveFrameNumber, MovieSceneValue);
 										
 					return KeyHandle != FKeyHandle::Invalid();
 				}
diff --git a/Engine/Source/Editor/AnimationModifiers/Private/AnimationModifier.cpp b/Engine/Source/Editor/AnimationModifiers/Private/AnimationModifier.cpp
--- a/Engine/Source/Editor/AnimationModifiers/Private/AnimationModifier.cpp
+++ b/Engine/Source/Editor/AnimationModifiers/Private/AnimationModifier.cpp
@@ -90,11 +90,8 @@ void UAnimationModifier::ApplyToAnimationSequence(class UAnimSequence* InAnimati
 
 	IAnimationDataController& Controller = CurrentAnimSequence->GetController(); //-V595
 
-	{
-		IAnimationDataController::FScopedBracket ScopedBracket(Controller, LOCTEXT("ApplyModifierBracket", "Applying Animation Modifier"));
-		/** Reverting and applying, populates the log with possible warnings and or errors to notify the user about */
-		OnApply(CurrentAnimSequence);
-	}
+	/** Reverting and applying, populates the log with possible warnings and or errors to notify the user about */
+	OnApply(CurrentAnimSequence);
 
 	// Apply transaction
 	ModifierTransaction.BeginOperation();
