Remove empty tracks to ensure track key count is consistent with data model

diff --git a/Engine/Plugins/Animation/AnimationData/Source/AnimationData/Private/AnimSequencerDataModel.cpp b/Engine/Plugins/Animation/AnimationData/Source/AnimationData/Private/AnimSequencerDataModel.cpp
--- a/Engine/Plugins/Animation/AnimationData/Source/AnimationData/Private/AnimSequencerDataModel.cpp
+++ b/Engine/Plugins/Animation/AnimationData/Source/AnimationData/Private/AnimSequencerDataModel.cpp
@@ -879,8 +879,8 @@ void UAnimationSequencerDataModel::GenerateLegacyBoneData()
 						}
 					});
 
-					// Remove any empty non-populate tracks
-					//LegacyBoneAnimationTracks.RemoveAll([](const FBoneAnimationTrack& Track) { return Track.InternalTrackData.PosKeys.Num() == 0; } );
+					// Remove tracks for bones not present in the new hierarchy
+					LegacyBoneAnimationTracks.RemoveAll([](const FBoneAnimationTrack& Track) { return Track.Name.IsNone() && Track.InternalTrackData.PosKeys.Num() == 0; } );
 				}
 				else
 				{						
