Fix for reimporting not clearing the cached DDC key GUID (used to limit need for recompression due to floating point differences after initial upgrade to DataModel)

diff --git a/Engine/Plugins/Animation/AnimationData/Source/AnimationData/Private/AnimSequencerDataModel.cpp b/Engine/Plugins/Animation/AnimationData/Source/AnimationData/Private/AnimSequencerDataModel.cpp
--- a/Engine/Plugins/Animation/AnimationData/Source/AnimationData/Private/AnimSequencerDataModel.cpp
+++ b/Engine/Plugins/Animation/AnimationData/Source/AnimationData/Private/AnimSequencerDataModel.cpp
@@ -20,6 +20,7 @@
 #include "Animation/AnimationSettings.h"
 #include "MovieScene/Private/Channels/MovieSceneCurveChannelImpl.h"
 #include "Compilation/MovieSceneCompiledDataManager.h"
+#include "UObject/UObjectThreadContext.h"
 
 #define LOCTEXT_NAMESPACE "AnimSequencerDataModel"
 
@@ -642,7 +643,8 @@ void UAnimationSequencerDataModel::OnNotify(const EAnimDataModelNotifyType& Noti
 		// Once the model has been populated and a modification is made - invalidate the cached GUID
 		auto ResetCachedGUID = [this]()
 		{
-			if (CachedRawDataGUID.IsValid() && !Collector.Contains(EAnimDataModelNotifyType::Populated))
+			// Prevent reset when being populated inside of upgrade path (always happens in UAnimSequenceBase::PostLoad)
+			if (CachedRawDataGUID.IsValid() && (!Collector.Contains(EAnimDataModelNotifyType::Populated) || !FUObjectThreadContext::Get().IsRoutingPostLoad))
 			{
 				CachedRawDataGUID.Invalidate();			
 			}
