Animation Compression
- Added new compression and 'compiling' API/path
  * Relies on newer DDC API
  * Only fetches compressible data when data is not found in DDC
  * Fetching data happens off the GT now (including Additive animation)
- Deprecated and replace AnimSequence API around compression
** AnimStreamable still relies on old, synchronous, compression path **

IAnimationDataModel
- Removed bone track data being stored as FBoneAnimationTrack
- Bone animation can now only be referenced by Name (not track index anymore)
- Deprecated any API relying on FBoneAnimationTrack
- Added API to retrieve FTransform(s) for given frame(s) with provided bone name
- Added API to lock model against modifications during evaluation (required for non-racing off-GT evaluation)

Animation Sequence
- Frame-rate is now stored on a per-platform basis, allowing for future replacement of frame-stripping (editing property is disabled for now)
- Now allows for storing _transient_ per-platform compressed animation data (required for multi-platform cook)

PerPlatformProperties
- Added per-platform FFrameRate implementation

diff --git a/Engine/Plugins/Animation/AnimationData/Source/AnimationData/Private/AnimSequencerController.cpp b/Engine/Plugins/Animation/AnimationData/Source/AnimationData/Private/AnimSequencerController.cpp
--- a/Engine/Plugins/Animation/AnimationData/Source/AnimationData/Private/AnimSequencerController.cpp
+++ b/Engine/Plugins/Animation/AnimationData/Source/AnimationData/Private/AnimSequencerController.cpp
@@ -21,6 +21,7 @@
 
 #include "AnimSequencerHelpers.h"
 #include "Animation/AnimationSettings.h"
+#include "UObject/UE5MainStreamObjectVersion.h"
 
 #define LOCTEXT_NAMESPACE "AnimSequencerController"
 
@@ -114,6 +115,7 @@ void UAnimSequencerController::ResizeNumberOfFrames(FFrameNumber NewLength, FFra
 				// Ensure that the start and end length of either removal or insertion are valid
 				if (T0 < T1)
 				{
+					IAnimationDataModel::FEvaluationAndModificationLock Lock(*ModelInterface);
 					FTransaction Transaction = ConditionalTransaction(LOCTEXT("ResizePlayLength", "Resizing Play Length"), bShouldTransact);
 					
 					const TObjectPtr<UMovieScene>& MovieScene = Model->MovieScene;
@@ -176,6 +178,7 @@ void UAnimSequencerController::ResizeInFrames(FFrameNumber NewLength, FFrameNumb
 				// Ensure that the start and end length of either removal or insertion are valid
 				if (T0 < T1)
 				{
+					IAnimationDataModel::FEvaluationAndModificationLock Lock(*ModelInterface);
 					FBracket Bracket = ConditionalBracket(LOCTEXT("ResizeModel", "Resizing Animation Data"), bShouldTransact);
 					const bool bInserted = NewLength > CurrentNumberOFrames;
 					ResizeNumberOfFrames(NewLength, T0, T1, bShouldTransact);
@@ -285,6 +288,7 @@ void UAnimSequencerController::SetFrameRate(FFrameRate FrameRate, bool bShouldTr
 
 		if (FrameRate.IsMultipleOf(CurrentFrameRate) || FrameRate.IsFactorOf(CurrentFrameRate) || !Model->bPopulated)
 		{
+			IAnimationDataModel::FEvaluationAndModificationLock Lock(*ModelInterface);
 			FTransaction Transaction = ConditionalTransaction(LOCTEXT("SetFrameRate", "Setting Frame Rate"), bShouldTransact);
 			ConditionalAction<UE::Anim::FSetFrameRateAction>(bShouldTransact, Model.Get());
 
@@ -389,6 +393,7 @@ void UAnimSequencerController::FindOrAddCurveNamesOnSkeleton(USkeleton* Skeleton
 	{
 		if (IsSupportedCurveType(SupportedCurveType))
 		{
+			IAnimationDataModel::FEvaluationAndModificationLock Lock(*ModelInterface);
 			FBracket Bracket = ConditionalBracket(LOCTEXT("FindOrAddRawCurveNames", "Updating Skeleton with Animation Curve Names"), bShouldTransact);
 			switch (SupportedCurveType)
 			{
@@ -467,6 +472,7 @@ bool UAnimSequencerController::RemoveBoneTracksMissingFromSkeleton(const USkelet
 			TArray<FName> TracksUpdated;
 			const FReferenceSkeleton& ReferenceSkeleton = Skeleton->GetReferenceSkeleton();
 
+			IAnimationDataModel::FEvaluationAndModificationLock Lock(*ModelInterface);
 			if (UMovieSceneControlRigParameterSection* Section = Model->GetFKControlRigSection())
 			{
 				const TArray<FTransformParameterNameAndCurves>& BoneTrackCurves = Section->GetTransformParameterNamesAndCurves();
@@ -539,6 +545,7 @@ void UAnimSequencerController::UpdateAttributesFromSkeleton(const USkeleton* Ske
 
 		if (ToRemoveIdentifiers.Num() || ToDuplicateIdentifiers.Num())
 		{
+			IAnimationDataModel::FEvaluationAndModificationLock Lock(*ModelInterface);
 			FBracket Bracket = ConditionalBracket(LOCTEXT("VerifyAttributeBoneNames", "Remapping Animation Attribute Data"), bShouldTransact);
 			for (const FAnimationAttributeIdentifier& Identifier : ToRemoveIdentifiers)
 			{
@@ -567,6 +574,7 @@ void UAnimSequencerController::ResetModel(bool bShouldTransact /*= true*/)
 
 	if (Model->bPopulated)
 	{
+		IAnimationDataModel::FEvaluationAndModificationLock Lock(*ModelInterface);
 		FBracket Bracket = ConditionalBracket(LOCTEXT("ResetModel", "Clearing Animation Data"), bShouldTransact);
 
 		RemoveAllBoneTracks(bShouldTransact);
@@ -593,6 +601,7 @@ bool UAnimSequencerController::AddCurve(const FAnimationCurveIdentifier& CurveId
 		{
 			if (!Model->FindCurve(CurveId))
 			{
+				IAnimationDataModel::FEvaluationAndModificationLock Lock(*ModelInterface);
 				FTransaction Transaction = ConditionalTransaction(LOCTEXT("AddRawCurve", "Adding Animation Curve"), bShouldTransact);
 
 				FCurveAddedPayload Payload;
@@ -665,6 +674,7 @@ bool UAnimSequencerController::DuplicateCurve(const FAnimationCurveIdentifier& C
 				{
 					if (!Model->FindCurve(NewCurveId))
 					{
+						IAnimationDataModel::FEvaluationAndModificationLock Lock(*ModelInterface);
 						FTransaction Transaction = ConditionalTransaction(LOCTEXT("CopyRawCurve", "Duplicating Animation Curve"), bShouldTransact);
 
 						auto DuplicateCurve = [NewCurveId, NewCurveName = NewCurveId.InternalName, this](auto& LegacyCurveDataArray, const auto& SourceCurve)
@@ -739,6 +749,7 @@ bool UAnimSequencerController::RemoveCurve(const FAnimationCurveIdentifier& Curv
 		{
 			if (Model->FindCurve(CurveId) != nullptr)
 			{
+				IAnimationDataModel::FEvaluationAndModificationLock Lock(*ModelInterface);
 				FTransaction Transaction = ConditionalTransaction(LOCTEXT("RemoveCurve", "Removing Animation Curve"), bShouldTransact);
 
 				switch (SupportedCurveType)
@@ -848,6 +859,7 @@ bool UAnimSequencerController::SetCurveFlag(const FAnimationCurveIdentifier& Cur
 	
 	if (Curve)
 	{
+		IAnimationDataModel::FEvaluationAndModificationLock Lock(*ModelInterface);
 		FTransaction Transaction = ConditionalTransaction(LOCTEXT("SetCurveFlag", "Setting Raw Curve Flag"), bShouldTransact);
 
 		const int32 CurrentFlags = Curve->GetCurveTypeFlags();
@@ -894,6 +906,7 @@ bool UAnimSequencerController::SetCurveFlags(const FAnimationCurveIdentifier& Cu
 
 	if (Curve)
 	{
+		IAnimationDataModel::FEvaluationAndModificationLock Lock(*ModelInterface);
 		FTransaction Transaction = ConditionalTransaction(LOCTEXT("SetCurveFlags", "Setting Raw Curve Flags"), bShouldTransact);
 
 		const int32 CurrentFlags = Curve->GetCurveTypeFlags();
@@ -928,6 +941,7 @@ bool UAnimSequencerController::SetTransformCurveKeys(const FAnimationCurveIdenti
 	{
 		if (Model->FindMutableTransformCurveById(CurveId) != nullptr)
 		{
+			IAnimationDataModel::FEvaluationAndModificationLock Lock(*ModelInterface);
 			FBracket Bracket = ConditionalBracket(LOCTEXT("SetTransformCurveKeys_Bracket", "Setting Transform Curve Keys"), bShouldTransact);
 			
 			struct FKeys
@@ -1008,6 +1022,7 @@ bool UAnimSequencerController::SetTransformCurveKey(const FAnimationCurveIdentif
 
 	if (Model->FindMutableTransformCurveById(CurveId) != nullptr)
 	{
+		IAnimationDataModel::FEvaluationAndModificationLock Lock(*ModelInterface);
 		FBracket Bracket = ConditionalBracket(LOCTEXT("AddTransformCurveKey_Bracket", "Setting Transform Curve Key"), bShouldTransact);
 		struct FKeys
 		{
@@ -1066,6 +1081,7 @@ bool UAnimSequencerController::RemoveTransformCurveKey(const FAnimationCurveIden
 		const TArray<FString> SubCurveNames = { TEXT( "Translation"), TEXT( "Rotation"), TEXT( "Scale") };
 		const TArray<FString> ChannelCurveNames = { TEXT("X"), TEXT("Y"), TEXT("Z") };
 
+		IAnimationDataModel::FEvaluationAndModificationLock Lock(*ModelInterface);
 		FBracket Bracket = ConditionalBracket(LOCTEXT("RemoveTransformCurveKey_Bracket", "Deleting Animation Transform Curve Key"), bShouldTransact);
 		
 		for (int32 SubCurveIndex = 0; SubCurveIndex < 3; ++SubCurveIndex)
@@ -1102,6 +1118,7 @@ bool UAnimSequencerController::RenameCurve(const FAnimationCurveIdentifier& Curv
 			{
 				if (FAnimCurveBase* Curve = Model->FindMutableCurveById(CurveToRenameId))
 				{
+					IAnimationDataModel::FEvaluationAndModificationLock Lock(*ModelInterface);
 					FTransaction Transaction = ConditionalTransaction(LOCTEXT("RenameCurve", "Renaming Curve"), bShouldTransact);
 
 					FCurveRenamedPayload Payload;
@@ -1164,6 +1181,7 @@ bool UAnimSequencerController::SetCurveColor(const FAnimationCurveIdentifier& Cu
 		{
 			if (const FFloatCurve* Curve = Model->FindMutableFloatCurveById(CurveId))
 			{
+				IAnimationDataModel::FEvaluationAndModificationLock Lock(*ModelInterface);
 				FTransaction Transaction = ConditionalTransaction(LOCTEXT("ChangingCurveColor", "Changing Curve Color"), bShouldTransact);
 
 				ConditionalAction<UE::Anim::FSetCurveColorAction>(bShouldTransact,  CurveId, Curve->Color);
@@ -1204,6 +1222,7 @@ bool UAnimSequencerController::ScaleCurve(const FAnimationCurveIdentifier& Curve
 	{
 		if (FFloatCurve* Curve = Model->FindMutableFloatCurveById(CurveId))
 		{
+			IAnimationDataModel::FEvaluationAndModificationLock Lock(*ModelInterface);
 			FTransaction Transaction = ConditionalTransaction(LOCTEXT("ScalingCurve", "Scaling Curve"), bShouldTransact);
 
 			Curve->FloatCurve.ScaleCurve(Origin, Factor);
@@ -1258,6 +1277,7 @@ bool UAnimSequencerController::SetCurveKey(const FAnimationCurveIdentifier& Curv
 		const FKeyHandle Handle = RichCurve->FindKey(Key.Time, 0.f);
 		if (Handle != FKeyHandle::Invalid())
 		{
+			IAnimationDataModel::FEvaluationAndModificationLock Lock(*ModelInterface);
 			FTransaction Transaction = ConditionalTransaction(LOCTEXT("SetNamedCurveKey", "Setting Curve Key"), bShouldTransact);
 			// Cache old value for action
 			const FRichCurveKey CurrentKey = RichCurve->GetKey(Handle);
@@ -1306,6 +1326,7 @@ bool UAnimSequencerController::RemoveCurveKey(const FAnimationCurveIdentifier& C
 		const FKeyHandle Handle = RichCurve->FindKey(Time, 0.f);
 		if (Handle != FKeyHandle::Invalid())
 		{
+			IAnimationDataModel::FEvaluationAndModificationLock Lock(*ModelInterface);
 			FTransaction Transaction = ConditionalTransaction(LOCTEXT("RemoveNamedCurveKey", "Removing Curve Key"), bShouldTransact);
 
 			// Cached current value for action
@@ -1345,6 +1366,7 @@ bool UAnimSequencerController::SetCurveKeys(const FAnimationCurveIdentifier& Cur
 		bool bKeysSet = false;
 		if (FRichCurve* RichCurve = Model->GetMutableRichCurve(CurveId))
 		{
+			IAnimationDataModel::FEvaluationAndModificationLock Lock(*ModelInterface);
 			FTransaction Transaction = ConditionalTransaction(LOCTEXT("SettingNamedCurveKeys", "Setting Curve Keys"), bShouldTransact);
 			ConditionalAction<UE::Anim::FSetRichCurveKeysAction>(bShouldTransact,  CurveId, RichCurve->GetConstRefOfKeys());
 
@@ -1394,6 +1416,7 @@ bool UAnimSequencerController::SetCurveAttributes(const FAnimationCurveIdentifie
 	const FRichCurve* RichCurve = Model->GetMutableRichCurve(CurveId);
 	if (RichCurve)
 	{
+		IAnimationDataModel::FEvaluationAndModificationLock Lock(*ModelInterface);
 		FTransaction Transaction = ConditionalTransaction(LOCTEXT("SettingNamedCurveAttributes", "Setting Curve Attributes"), bShouldTransact);
 
 		FCurveAttributes CurrentAttributes;
@@ -1480,15 +1503,15 @@ void UAnimSequencerController::NotifyBracketClosed()
 	Model->GetNotifier().Notify(EAnimDataModelNotifyType::BracketClosed);
 }
 
-int32 UAnimSequencerController::AddBoneTrack(FName BoneName, bool bShouldTransact /*= true*/)
+bool UAnimSequencerController::AddBoneCurve(FName BoneName, bool bShouldTransact /*= true*/)
 {
 	if (!ModelInterface->GetAnimationSequence())
 	{
-		return INDEX_NONE;
+		return false;
 	}
 
 	FTransaction Transaction = ConditionalTransaction(LOCTEXT("AddBoneTrack", "Adding Animation Data Track"), bShouldTransact);
-	return InsertBoneTrack(BoneName, INDEX_NONE, bShouldTransact);
+	return InsertBoneTrack(BoneName, INDEX_NONE, bShouldTransact) != INDEX_NONE;
 }
 
 int32 UAnimSequencerController::InsertBoneTrack(FName BoneName, int32 DesiredIndex, bool bShouldTransact /*= true*/)
@@ -1497,10 +1520,8 @@ int32 UAnimSequencerController::InsertBoneTrack(FName BoneName, int32 DesiredInd
 	{
 		return INDEX_NONE;
 	}
-	
-	const int32 TrackIndex = Model->GetBoneTrackIndexByName(BoneName);
 
-	if (TrackIndex == INDEX_NONE)
+	if (!Model->IsValidBoneTrackName(BoneName))
 	{
 		QUICK_SCOPE_CYCLE_COUNTER(STAT_InsertBoneTrack);
 		if (Model->GetNumBoneTracks() >= MAX_ANIMATION_TRACKS)
@@ -1511,7 +1532,6 @@ int32 UAnimSequencerController::InsertBoneTrack(FName BoneName, int32 DesiredInd
 		else
 		{
 			// Determine correct index to do insertion at
-			const int32 InsertIndex = Model->LegacyBoneAnimationTracks.IsValidIndex(DesiredIndex) ? DesiredIndex : Model->LegacyBoneAnimationTracks.Num();
 			int32 BoneIndex = INDEX_NONE;
 
 			if (const UAnimSequence* AnimationSequence = Model->GetAnimationSequence())
@@ -1537,25 +1557,21 @@ int32 UAnimSequencerController::InsertBoneTrack(FName BoneName, int32 DesiredInd
 
 			if (BoneIndex != INDEX_NONE || IgnoreSkeletonValidation())
 			{
+				IAnimationDataModel::FEvaluationAndModificationLock Lock(*ModelInterface);
 				FTransaction Transaction = ConditionalTransaction(LOCTEXT("InsertBoneTrack", "Inserting Animation Data Track"), bShouldTransact);
-
-				FBoneAnimationTrack& NewTrack = Model->LegacyBoneAnimationTracks.InsertDefaulted_GetRef(InsertIndex);
-				NewTrack.Name = BoneName;
-				NewTrack.BoneTreeIndex = BoneIndex;
 				
 				FAnimationTrackAddedPayload Payload;
 				Payload.Name = BoneName;
-				Payload.TrackIndex = InsertIndex;
 
 				if(!AddBoneControl(BoneName))
 				{
 					ReportErrorf(LOCTEXT("FailedToAddBoneControlElement", "Failed to add Bone control with name {0}"), FText::FromName(BoneName));
 				}
 				
-				ConditionalAction<UE::Anim::FRemoveTrackAction>(bShouldTransact,  BoneName);
+				ConditionalAction<UE::Anim::FRemoveTrackAction>(bShouldTransact, BoneName);
 				Model->GetNotifier().Notify<FAnimationTrackAddedPayload>(EAnimDataModelNotifyType::TrackAdded, Payload);
 
-				return InsertIndex;
+				return 0;
 			}
 		}
 	}
@@ -1564,7 +1580,7 @@ int32 UAnimSequencerController::InsertBoneTrack(FName BoneName, int32 DesiredInd
 		ReportWarningf(LOCTEXT("TrackNameAlreadyExistsWarning", "Track with name {0} already exists"), FText::FromName(BoneName));
 	}
 	
-	return TrackIndex;
+	return INDEX_NONE;
 }
 
 bool UAnimSequencerController::RemoveBoneTrack(FName BoneName, bool bShouldTransact /*= true*/)
@@ -1574,21 +1590,18 @@ bool UAnimSequencerController::RemoveBoneTrack(FName BoneName, bool bShouldTrans
 		return false;
 	}
 
-	const FBoneAnimationTrack* ExistingTrackPtr = Model->FindBoneTrackByName(BoneName);
-
-	if (ExistingTrackPtr != nullptr)
+	if (Model->IsValidBoneTrackName(BoneName))
 	{
+		IAnimationDataModel::FEvaluationAndModificationLock Lock(*ModelInterface);
 		FTransaction Transaction = ConditionalTransaction(LOCTEXT("RemoveBoneTrack", "Removing Animation Data Track"), bShouldTransact);
-		const int32 TrackIndex = Model->LegacyBoneAnimationTracks.IndexOfByPredicate([ExistingTrackPtr](const FBoneAnimationTrack& Track)
-		{
-			return Track.Name == ExistingTrackPtr->Name;
-		});
-
-		ensure(TrackIndex != INDEX_NONE);
-
-		ConditionalAction<UE::Anim::FAddTrackAction>(bShouldTransact, *ExistingTrackPtr);
-		Model->LegacyBoneAnimationTracks.RemoveAt(TrackIndex);
 
+		if (bShouldTransact)
+		{
+			TArray<FTransform> BoneTransforms;
+			Model->GetBoneTrackTransforms(BoneName, BoneTransforms);
+			ConditionalAction<UE::Anim::FAddTrackAction>(bShouldTransact, BoneName, MoveTemp(BoneTransforms));			
+		}
+	
 		if(!RemoveBoneControl(BoneName))
 		{
 			ReportError(LOCTEXT("FailedToRemoveBoneControl", "Failed to remove bone control"));
@@ -1642,41 +1655,54 @@ bool UAnimSequencerController::SetBoneTrackKeys(FName BoneName, const TArray<FVe
 
 	if (MaxNumKeys > 0)
 	{
-		QUICK_SCOPE_CYCLE_COUNTER(STAT_SetBoneTrackKeys);
 		const bool bValidPosKeys = PositionalKeys.Num() == MaxNumKeys;
 		const bool bValidRotKeys = RotationalKeys.Num() == MaxNumKeys;
 		const bool bValidScaleKeys = ScalingKeys.Num() == MaxNumKeys;
 
 		if (bValidPosKeys && bValidRotKeys && bValidScaleKeys)
 		{
-			if (FBoneAnimationTrack* TrackPtr = Model->FindMutableBoneTrackByName(BoneName))
+			if (UMovieSceneControlRigParameterSection* Section = Model->GetFKControlRigSection())
 			{
-				static const FText TransactionText = LOCTEXT("SetTrackKeysTransaction", "Setting Animation Data Track keys");
-				FTransaction Transaction = ConditionalTransaction(TransactionText, bShouldTransact);
-				ConditionalAction<UE::Anim::FSetTrackKeysAction>(bShouldTransact,  *TrackPtr);
-
-				if(UAnimationSequencerDataModel::ValidationMode == 1)
+				if (UControlRig* ControlRig = Model->GetControlRig())
 				{
-					TrackPtr->InternalTrackData.PosKeys = PositionalKeys;
-					TrackPtr->InternalTrackData.RotKeys = RotationalKeys;
-					TrackPtr->InternalTrackData.ScaleKeys = ScalingKeys;
-				}
+					if (URigHierarchy* Hierarchy = ControlRig->GetHierarchy())
+					{
+						const FRigElementKey BoneKey(BoneName, ERigElementType::Bone);
+						const FRigElementKey BoneNameControlKey(UFKControlRig::GetControlName(BoneName, ERigElementType::Bone), ERigElementType::Control);
+						
+						const bool bContainsBone = Hierarchy->Contains(BoneKey);
+						const bool bContainsBoneControl = Hierarchy->Contains(BoneNameControlKey);
+						if (bContainsBone && bContainsBoneControl)
+						{
+							IAnimationDataModel::FEvaluationAndModificationLock Lock(*ModelInterface);
+							FTransaction Transaction = ConditionalTransaction(LOCTEXT("SetTrackKeysTransaction", "Setting Animation Data Track keys"), bShouldTransact);
 
-				if(!SetBoneCurveKeys(BoneName, PositionalKeys, RotationalKeys, ScalingKeys))
-				{
-					ReportError(LOCTEXT("FailedToSetBoneCurveKeys", "Failed to set bone curve keys"));
-				}
+							if (bShouldTransact)
+							{
+								TArray<FTransform> BoneTransforms;
+								Model->GetBoneTrackTransforms(BoneName, BoneTransforms);
+								ConditionalAction<UE::Anim::FSetTrackKeysAction>(bShouldTransact, BoneName, BoneTransforms);							
+							}
 
-				FAnimationTrackChangedPayload Payload;
-				Payload.Name = BoneName;
+							if(!SetBoneCurveKeys(BoneName, PositionalKeys, RotationalKeys, ScalingKeys))
+							{
+								ReportError(LOCTEXT("FailedToSetBoneCurveKeys", "Failed to set bone curve keys"));
+							}
 
-				Model->GetNotifier().Notify(EAnimDataModelNotifyType::TrackChanged, Payload);
+							FAnimationTrackChangedPayload Payload;
+							Payload.Name = BoneName;
 
-				return true;
-			}
-			else
-			{
-				ReportWarningf(LOCTEXT("InvalidTrackNameWarning", "Track with name {0} does not exist"), FText::FromName(BoneName));
+							Model->GetNotifier().Notify(EAnimDataModelNotifyType::TrackChanged, Payload);
+
+							return true;
+						}
+					
+					}
+					else
+					{
+						ReportWarningf(LOCTEXT("InvalidTrackNameWarning", "Track with name {0} does not exist"), FText::FromName(BoneName));
+					}
+				}
 			}
 		}
 		else
@@ -1749,50 +1775,42 @@ bool UAnimSequencerController::UpdateBoneTrackKeys(FName BoneName, const FInt32R
 			const int32 NumKeysToSet = RangeMax - RangeMin;
 			if (NumKeysToSet == MaxNumKeys)
 			{
-				if (FBoneAnimationTrack* TrackPtr = Model->FindMutableBoneTrackByName(BoneName))
+				if (Model->IsValidBoneTrackName(BoneName))
 				{
-					const FInt32Range TrackKeyRange(0, TrackPtr->InternalTrackData.PosKeys.Num());
+					const FInt32Range TrackKeyRange(0, Model->GetNumberOfKeys());
 					if(TrackKeyRange.Contains(KeyRangeToSet))
 					{
-						FRawAnimSequenceTrack& InternalTrackData = TrackPtr->InternalTrackData;
-						TArray<FVector3f>& TrackPosKeys = InternalTrackData.PosKeys;
-						TArray<FQuat4f>& TrackRotKeys = InternalTrackData.RotKeys;
-						TArray<FVector3f>& TrackScaleKeys = InternalTrackData.ScaleKeys;
-
+						IAnimationDataModel::FEvaluationAndModificationLock Lock(*ModelInterface);
 						FTransaction Transaction = ConditionalTransaction(LOCTEXT("SetTrackKeysRangeTransaction", "Setting Animation Data Track keys"), bShouldTransact);
-						ConditionalAction<UE::Anim::FSetTrackKeysAction>(bShouldTransact, *TrackPtr);
 
+						if (bShouldTransact)
+						{
+							TArray<FTransform> BoneTransforms;
+							Model->GetBoneTrackTransforms(BoneName, BoneTransforms);
+							ConditionalAction<UE::Anim::FSetTrackKeysAction>(bShouldTransact, BoneName, BoneTransforms);
+						}
+						
+						TArray<FTransform> TransformKeys;
+						TArray<float> TimeValues;
+					
+						const FFrameRate& FrameRate = ModelInterface->GetFrameRate();
+						TransformKeys.Reserve(PositionalKeys.Num());
+						TimeValues.Reserve(PositionalKeys.Num());
 						int32 KeyIndex = 0;
 						for (int32 RangeIndex = RangeMin; RangeIndex < RangeMax; ++RangeIndex, ++KeyIndex)
 						{
-							TrackPosKeys[RangeIndex] = PositionalKeys[KeyIndex];
-							TrackRotKeys[RangeIndex] = RotationalKeys[KeyIndex];
-							TrackScaleKeys[RangeIndex] = ScalingKeys[KeyIndex];
+							TransformKeys.Add(FTransform(FQuat(RotationalKeys[KeyIndex]), FVector(PositionalKeys[KeyIndex]), FVector(ScalingKeys[KeyIndex])));
+							TimeValues.Add(FrameRate.AsSeconds(RangeIndex));
 						}
-
+					
+						if(!UpdateBoneCurveKeys(BoneName, TransformKeys, TimeValues))
 						{
-							TArray<FTransform> TransformKeys;
-							TArray<float> TimeValues;
-
-							const FFrameRate& FrameRate = ModelInterface->GetFrameRate();
-							TransformKeys.Reserve(PositionalKeys.Num());
-							TimeValues.Reserve(PositionalKeys.Num());
-							KeyIndex = 0;
-							for (int32 RangeIndex = RangeMin; RangeIndex < RangeMax; ++RangeIndex, ++KeyIndex)
-							{
-								TransformKeys.Add(FTransform(FQuat(RotationalKeys[KeyIndex]), FVector(PositionalKeys[KeyIndex]), FVector(ScalingKeys[KeyIndex])));
-								TimeValues.Add(FrameRate.AsSeconds(RangeIndex));
-							}
-
-							if(!UpdateBoneCurveKeys(BoneName, TransformKeys, TimeValues))
-							{
-								ReportError(LOCTEXT("FailedToUpdateBoneCurveKeys", "Failed to update bone curve keys"));
-							}
+							ReportError(LOCTEXT("FailedToUpdateBoneCurveKeys", "Failed to update bone curve keys"));
 						}
-
+						
 						FAnimationTrackChangedPayload Payload;
 						Payload.Name = BoneName;
-
+						
 						Model->GetNotifier().Notify(EAnimDataModelNotifyType::TrackChanged, Payload);
 
 						return true;
@@ -1908,6 +1926,7 @@ bool UAnimSequencerController::AddAttribute(const FAnimationAttributeIdentifier&
 
 		if (!bAttributeAlreadyExists)
 		{
+			IAnimationDataModel::FEvaluationAndModificationLock Lock(*ModelInterface);
 			FTransaction Transaction = ConditionalTransaction(LOCTEXT("AddAttribute", "Adding Animated Bone Attribute"), bShouldTransact);
 
 			FAnimatedBoneAttribute& Attribute = Model->AnimatedBoneAttributes.AddDefaulted_GetRef();
@@ -1948,6 +1967,7 @@ bool UAnimSequencerController::RemoveAttribute(const FAnimationAttributeIdentifi
 
 		if (AttributeIndex != INDEX_NONE)
 		{
+			IAnimationDataModel::FEvaluationAndModificationLock Lock(*ModelInterface);
 			FTransaction Transaction = ConditionalTransaction(LOCTEXT("RemoveAttribute", "Removing Animated Bone Attribute"), bShouldTransact);
 
 			ConditionalAction<UE::Anim::FAddAtributeAction>(bShouldTransact,  Model->AnimatedBoneAttributes[AttributeIndex]);
@@ -2040,6 +2060,7 @@ bool UAnimSequencerController::SetAttributeKey_Internal(const FAnimationAttribut
 			{
 				if (TypeStruct == AttributePtr->Identifier.GetType())
 				{
+					IAnimationDataModel::FEvaluationAndModificationLock Lock(*ModelInterface);
 					FTransaction Transaction = ConditionalTransaction(LOCTEXT("SettingAttributeKey", "Setting Animated Bone Attribute key"), bShouldTransact);
 
 					FAttributeCurve& Curve = AttributePtr->Curve;
@@ -2103,6 +2124,7 @@ bool UAnimSequencerController::SetAttributeKeys_Internal(const FAnimationAttribu
 			{
 				if (TypeStruct == AttributePtr->Identifier.GetType())
 				{
+					IAnimationDataModel::FEvaluationAndModificationLock Lock(*ModelInterface);
 					FTransaction Transaction = ConditionalTransaction(LOCTEXT("SettingAttributeKeys", "Setting Animated Bone Attribute keys"), bShouldTransact);
 
 					FAnimatedBoneAttribute& Attribute = *AttributePtr;
@@ -2160,6 +2182,7 @@ bool UAnimSequencerController::RemoveAttributeKey(const FAnimationAttributeIdent
 
 			if (KeyHandle != FKeyHandle::Invalid())
 			{
+				IAnimationDataModel::FEvaluationAndModificationLock Lock(*ModelInterface);
 				FTransaction Transaction = ConditionalTransaction(LOCTEXT("RemovingAttributeKey", "Removing Animated Bone Attribute key"), bShouldTransact);
 
 				ConditionalAction<UE::Anim::FAddAtributeKeyAction>(bShouldTransact,  AttributeIdentifier, Curve.GetKey(KeyHandle));
@@ -2204,6 +2227,7 @@ bool UAnimSequencerController::DuplicateAttribute(const FAnimationAttributeIdent
 			{
 				if(const FAnimatedBoneAttribute* AttributePtr = Model->FindAttribute(AttributeIdentifier))
 				{
+					IAnimationDataModel::FEvaluationAndModificationLock Lock(*ModelInterface);
 					FTransaction Transaction = ConditionalTransaction(LOCTEXT("DuplicateAttribute", "Duplicating Animation Attribute"), bShouldTransact);
 
 					FAnimatedBoneAttribute& DuplicateAttribute = Model->AnimatedBoneAttributes.AddDefaulted_GetRef();
@@ -2313,7 +2337,6 @@ void UAnimSequencerController::UpdateWithSkeleton(USkeleton* TargetSkeleton, boo
 		RemoveUnusedControlsAndCurves();		
 
 		// Forcefully re-generate legacy data structures
-		Model->GenerateLegacyBoneData();
 		Model->GenerateLegacyCurveData();
 	}
 	CloseBracket();
@@ -2326,6 +2349,7 @@ void UAnimSequencerController::PopulateWithExistingModel(TScriptInterface<IAnima
 	OpenBracket(LOCTEXT("PopulateWithExistingModel", "Copying Data from Legacy Animation Data Model"));
 	if (InModel.GetObject() && InModel.GetObject()->IsA<UAnimDataModel>())
 	{
+		IAnimationDataModel::FEvaluationAndModificationLock Lock(*ModelInterface);	
 		if (const UAnimSequence* AnimationSequence = Model->GetAnimationSequence())
 		{
 			PRAGMA_DISABLE_DEPRECATION_WARNINGS
@@ -2342,7 +2366,9 @@ void UAnimSequencerController::PopulateWithExistingModel(TScriptInterface<IAnima
 
 		SetFrameRate(ModelFrameRate, false);
 
+		PRAGMA_DISABLE_DEPRECATION_WARNINGS
 		const TArray<FBoneAnimationTrack>& BoneTracks = InModel->GetBoneAnimationTracks();
+		PRAGMA_ENABLE_DEPRECATION_WARNINGS
 		const bool bHasBoneTracks = BoneTracks.Num() > 0;
 
 		FFrameNumber NumberOfFrames = InModel->GetNumberOfFrames();
@@ -2375,7 +2401,7 @@ void UAnimSequencerController::PopulateWithExistingModel(TScriptInterface<IAnima
 				QUICK_SCOPE_CYCLE_COUNTER(STAT_PopulateBoneTrack);
 				for (const FBoneAnimationTrack& BoneTrack : BoneTracks)
 				{
-					AddBoneTrack(BoneTrack.Name, false);
+					AddBoneCurve(BoneTrack.Name, false);
 					SetBoneTrackKeys(BoneTrack.Name, BoneTrack.InternalTrackData.PosKeys, BoneTrack.InternalTrackData.RotKeys, BoneTrack.InternalTrackData.ScaleKeys, false);		
 				}
 			}
@@ -2452,73 +2478,83 @@ void UAnimSequencerController::DeclareConstructClasses(TArray<FTopLevelAssetPath
 
 void UAnimSequencerController::InitializeModel()
 {
-	if (Model->GetAnimationSequence())
+	if (const UAnimSequence* AnimSequence = Model->GetAnimationSequence())
 	{
-		if (UMovieScene* MovieScene = NewObject<UMovieScene>(Model.Get(), FName("MovieScene"), RF_Transactional))
+		if (Model->MovieScene == nullptr)
 		{
-			Model->MovieScene = MovieScene;
+			IAnimationDataModel::FEvaluationAndModificationLock Lock(*ModelInterface);
+			if (UMovieScene* MovieScene = NewObject<UMovieScene>(Model.Get(), FName("MovieScene"), RF_Transactional))
+			{
+				Model->MovieScene = MovieScene;
 
-			MovieScene->SetTickResolutionDirectly(UAnimationSettings::Get()->GetDefaultFrameRate());
-			MovieScene->SetDisplayRate(UAnimationSettings::Get()->GetDefaultFrameRate());
+				MovieScene->SetTickResolutionDirectly(UAnimationSettings::Get()->GetDefaultFrameRate());
+				MovieScene->SetDisplayRate(UAnimationSettings::Get()->GetDefaultFrameRate());
 
-			const TRange<FFrameNumber> DataRange = TRange<FFrameNumber>::Inclusive(FFrameNumber(0), 1);
-			MovieScene->SetPlaybackRange(DataRange);
-			
-			if(UMovieSceneControlRigParameterTrack* Track = MovieScene->AddTrack<UMovieSceneControlRigParameterTrack>())
-			{
-				if(UFKControlRig* ControlRig = NewObject<UFKControlRig>(Track, UFKControlRig::StaticClass(), NAME_None, RF_Transactional))
-				{
-					Model->InitializeFKControlRig(ControlRig, Model->GetSkeleton());
+				const TRange<FFrameNumber> DataRange = TRange<FFrameNumber>::Inclusive(FFrameNumber(0), 1);
+				MovieScene->SetPlaybackRange(DataRange);
 
-					// Remove all control elements (start fresh)
-					if (URigHierarchy* Hierarchy = ControlRig->GetHierarchy())
+				if (AnimSequence->GetLinkerCustomVersion(FUE5MainStreamObjectVersion::GUID) < FUE5MainStreamObjectVersion::IntroducingAnimationDataModel)
+				{
+					PRAGMA_DISABLE_DEPRECATION_WARNINGS
+					Model->CachedRawDataGUID = AnimSequence->GenerateGuidFromRawData();
+					PRAGMA_ENABLE_DEPRECATION_WARNINGS
+				}
+				
+				if(UMovieSceneControlRigParameterTrack* Track = MovieScene->AddTrack<UMovieSceneControlRigParameterTrack>())
+				{
+					if(UFKControlRig* ControlRig = NewObject<UFKControlRig>(Track, UFKControlRig::StaticClass(), NAME_None, RF_Transactional))
 					{
-						if (URigHierarchyController* Controller = Hierarchy->GetController())
+						Model->InitializeFKControlRig(ControlRig, Model->GetSkeleton());
+
+						// Remove all control elements (start fresh)
+						if (URigHierarchy* Hierarchy = ControlRig->GetHierarchy())
 						{
-							TArray<FRigElementKey> ControlElementKeysToRemove;
-							Hierarchy->ForEach<FRigControlElement>([&ControlElementKeysToRemove](const FRigControlElement* ControlElement) -> bool
+							if (URigHierarchyController* Controller = Hierarchy->GetController())
 							{
-								ControlElementKeysToRemove.Add(ControlElement->GetKey());
-								return true;
-							});
+								TArray<FRigElementKey> ControlElementKeysToRemove;
+								Hierarchy->ForEach<FRigControlElement>([&ControlElementKeysToRemove](const FRigControlElement* ControlElement) -> bool
+								{
+									ControlElementKeysToRemove.Add(ControlElement->GetKey());
+									return true;
+								});
 
-							Hierarchy->ForEach<FRigCurveElement>([&ControlElementKeysToRemove](const FRigCurveElement* CurveElement) -> bool
-							{
-								ControlElementKeysToRemove.Add(CurveElement->GetKey());							
-								return true;
-							});
+								Hierarchy->ForEach<FRigCurveElement>([&ControlElementKeysToRemove](const FRigCurveElement* CurveElement) -> bool
+								{
+									ControlElementKeysToRemove.Add(CurveElement->GetKey());							
+									return true;
+								});
 
-							for (const FRigElementKey& KeyToRemove : ControlElementKeysToRemove)
+								for (const FRigElementKey& KeyToRemove : ControlElementKeysToRemove)
+								{
+									Controller->RemoveElement(KeyToRemove);
+								}
+							}
+							else
 							{
-								Controller->RemoveElement(KeyToRemove);
+								ReportError(LOCTEXT("InvalidRigHierarchyController", "Non-valid RigHierarchyController found"));
 							}
+
+							ControlRig->RefreshActiveControls();
 						}
 						else
 						{
-							ReportError(LOCTEXT("InvalidRigHierarchyController", "Non-valid RigHierarchyController found"));
+							ReportError(LOCTEXT("InvalidRigHierarchy", "Unable to retrieve valid URigHierarchy"));
 						}
 
-						ControlRig->RefreshActiveControls();
+						UMovieSceneControlRigParameterSection* Section = Cast<UMovieSceneControlRigParameterSection>(Track->CreateControlRigSection(0, ControlRig, true));				
+						Section->SetRange(DataRange);
 					}
 					else
 					{
-						ReportError(LOCTEXT("InvalidRigHierarchy", "Unable to retrieve valid URigHierarchy"));
+						ReportError(LOCTEXT("FailedToCreateControlRig", "Failed to create new UFKControlRig"));
 					}
-
-					UMovieSceneControlRigParameterSection* Section = Cast<UMovieSceneControlRigParameterSection>(Track->CreateControlRigSection(0, ControlRig, true));				
-					Section->SetRange(DataRange);
+					
 				}
 				else
 				{
-					ReportError(LOCTEXT("FailedToCreateControlRig", "Failed to create new UFKControlRig"));
+					ReportError(LOCTEXT("FailedToCreateTrack", "Failed to create UMovieSceneControlRigParameterTrack"));
 				}
-				
 			}
-			else
-			{
-				ReportError(LOCTEXT("FailedToCreateTrack", "Failed to create UMovieSceneControlRigParameterTrack"));
-			}
-			
 		}
 		else
 		{
@@ -2579,6 +2615,19 @@ bool UAnimSequencerController::AddBoneControl(const FName& BoneName) const
 							Hierarchy->SetControlOffsetTransform(BoneControlElement, OffsetTransform, ERigTransformType::CurrentLocal, false, false, true);
 
 							Section->AddTransformParameter(BoneNameControlKey.Name, TOptional<FEulerTransform>(), false);
+
+							FTransformParameterNameAndCurves* ParameterCurvePair = Section->GetTransformParameterNamesAndCurves().FindByPredicate([BoneNameControlKey](const FTransformParameterNameAndCurves& Parameter)
+							{
+								return Parameter.ParameterName == BoneNameControlKey.Name;
+							});
+							check(ParameterCurvePair);
+
+							for (int32 Index = 0; Index < 3; ++Index)
+							{
+								ParameterCurvePair->Translation[Index].SetTickResolution(Model->GetFrameRate());
+								ParameterCurvePair->Rotation[Index].SetTickResolution(Model->GetFrameRate());
+								ParameterCurvePair->Scale[Index].SetTickResolution(Model->GetFrameRate());
+							}
 							
 							return true;
 						}
@@ -2709,7 +2758,6 @@ bool UAnimSequencerController::SetBoneCurveKeys(const FName& BoneName, const TAr
 					TArray<bool> AreKeysVarying;
 					AreKeysVarying.SetNumZeroed(9);
 					{
-						QUICK_SCOPE_CYCLE_COUNTER(STAT_SetBoneCurveKeys_ConstantCheck);
 						for (int32 KeyIndex = 1; KeyIndex < NumKeys; ++KeyIndex)
 						{
 							EulerAngles.Add(RotationalKeys[KeyIndex].Euler());
@@ -3129,7 +3177,7 @@ bool UAnimSequencerController::RenameCurveControl(const FName& CurveName, const
 									ControlElement->Settings.DisplayName = FName(*(NewCurveName.ToString() + TEXT(" Curve")));
 
 									// Rename the curve driving the control value
-									FScalarParameterNameAndCurve* ParameterCurvePair = Section->GetScalarParameterNamesAndCurves().FindByPredicate([CurveControlKey](FScalarParameterNameAndCurve Parameter)
+									FScalarParameterNameAndCurve* ParameterCurvePair = Section->GetScalarParameterNamesAndCurves().FindByPredicate([CurveControlKey](const FScalarParameterNameAndCurve& Parameter)
 									{
 										return Parameter.ParameterName == CurveControlKey.Name;
 									});
@@ -3361,7 +3409,7 @@ bool UAnimSequencerController::SetCurveControlKey(const FName& CurveName, const
 						Section->AddScalarParameter(CurveControlKey.Name, TOptional<float>(), true);
 					}
 
-					FScalarParameterNameAndCurve* ParameterCurvePair = Section->GetScalarParameterNamesAndCurves().FindByPredicate([CurveControlKey](FScalarParameterNameAndCurve Parameter)
+					FScalarParameterNameAndCurve* ParameterCurvePair = Section->GetScalarParameterNamesAndCurves().FindByPredicate([CurveControlKey](const FScalarParameterNameAndCurve& Parameter)
 					{
 						return Parameter.ParameterName == CurveControlKey.Name;
 					});
@@ -3432,7 +3480,7 @@ bool UAnimSequencerController::RemoveCurveControlKey(const FName& CurveName, flo
 				{
 					if(Section->HasScalarParameter(CurveControlKey.Name))
 					{
-						FScalarParameterNameAndCurve* ParameterCurvePair = Section->GetScalarParameterNamesAndCurves().FindByPredicate([CurveControlKey](FScalarParameterNameAndCurve Parameter)
+						FScalarParameterNameAndCurve* ParameterCurvePair = Section->GetScalarParameterNamesAndCurves().FindByPredicate([CurveControlKey](const FScalarParameterNameAndCurve& Parameter)
 						{
 							return Parameter.ParameterName == CurveControlKey.Name;
 						});
@@ -3506,7 +3554,7 @@ bool UAnimSequencerController::DuplicateCurveControl(const FName& CurveName, con
 						HierarchyController->AddControl(NewCurveControlKey.Name, FRigElementKey(), Settings, FRigControlValue::Make(CurveElement->Value), FTransform::Identity, FTransform::Identity, false);
 			
 						// Rename the curve driving the control value
-						const FScalarParameterNameAndCurve* ParameterCurvePair = Section->GetScalarParameterNamesAndCurves().FindByPredicate([CurveName](FScalarParameterNameAndCurve Parameter)
+						const FScalarParameterNameAndCurve* ParameterCurvePair = Section->GetScalarParameterNamesAndCurves().FindByPredicate([CurveName](const FScalarParameterNameAndCurve& Parameter)
 						{
 							return Parameter.ParameterName == UFKControlRig::GetControlName(CurveName, ERigElementType::Curve);
 						});
@@ -3515,7 +3563,7 @@ bool UAnimSequencerController::DuplicateCurveControl(const FName& CurveName, con
 						{
 							Section->AddScalarParameter(NewCurveControlKey.Name, TOptional<float>(), true);
 
-							FScalarParameterNameAndCurve* DuplicatedParameterCurvePair = Section->GetScalarParameterNamesAndCurves().FindByPredicate([NewCurveControlKey](FScalarParameterNameAndCurve Parameter)
+							FScalarParameterNameAndCurve* DuplicatedParameterCurvePair = Section->GetScalarParameterNamesAndCurves().FindByPredicate([NewCurveControlKey](const FScalarParameterNameAndCurve& Parameter)
 							{
 								return Parameter.ParameterName == NewCurveControlKey.Name;
 							});
diff --git a/Engine/Plugins/Animation/AnimationData/Source/AnimationData/Private/AnimSequencerController.h b/Engine/Plugins/Animation/AnimationData/Source/AnimationData/Private/AnimSequencerController.h
--- a/Engine/Plugins/Animation/AnimationData/Source/AnimationData/Private/AnimSequencerController.h
+++ b/Engine/Plugins/Animation/AnimationData/Source/AnimationData/Private/AnimSequencerController.h
@@ -39,7 +39,7 @@ public:
 	virtual void ResizePlayLength(float NewLength, float T0, float T1, bool bShouldTransact = true) override;
 	virtual void Resize(float Length, float T0, float T1, bool bShouldTransact = true) override;
 	PRAGMA_ENABLE_DEPRECATION_WARNINGS
-	virtual int32 AddBoneTrack(FName BoneName, bool bShouldTransact = true) override;
+	virtual bool AddBoneCurve(FName BoneName, bool bShouldTransact = true) override;
 	virtual int32 InsertBoneTrack(FName BoneName, int32 DesiredIndex, bool bShouldTransact = true) override;
 	virtual bool RemoveBoneTrack(FName BoneName, bool bShouldTransact = true) override;
 	virtual void RemoveAllBoneTracks(bool bShouldTransact = true) override;
diff --git a/Engine/Plugins/Animation/AnimationData/Source/AnimationData/Private/AnimSequencerDataModel.cpp b/Engine/Plugins/Animation/AnimationData/Source/AnimationData/Private/AnimSequencerDataModel.cpp
--- a/Engine/Plugins/Animation/AnimationData/Source/AnimationData/Private/AnimSequencerDataModel.cpp
+++ b/Engine/Plugins/Animation/AnimationData/Source/AnimationData/Private/AnimSequencerDataModel.cpp
@@ -109,6 +109,8 @@ void UAnimationSequencerDataModel::InitializeFKControlRig(UFKControlRig* FKContr
 	checkf(FKControlRig, TEXT("Invalid FKControlRig provided"));
 	if (Skeleton)
 	{
+		LockEvaluationAndModification();
+		
 		FKControlRig->SetObjectBinding(MakeShared<FControlRigObjectBinding>());
 		FKControlRig->GetObjectBinding()->BindToObject(Skeleton);
 
@@ -134,6 +136,8 @@ void UAnimationSequencerDataModel::InitializeFKControlRig(UFKControlRig* FKContr
 		FKControlRig->SetApplyMode(UseDirectFKControlRigMode == 1 ? EControlRigFKRigExecuteMode::Direct : EControlRigFKRigExecuteMode::Replace);
 		FKControlRig->SetBoneInitialTransformsFromRefSkeleton(Skeleton->GetReferenceSkeleton());
 		FKControlRig->Evaluate_AnyThread();
+		
+		UnlockEvaluationAndModification();
 	}
 }
 
@@ -214,7 +218,6 @@ void UAnimationSequencerDataModel::PreSave(FObjectPreSaveContext ObjectSaveConte
 void UAnimationSequencerDataModel::WillNeverCacheCookedPlatformDataAgain()
 {
 	Super::WillNeverCacheCookedPlatformDataAgain();
-	LegacyBoneAnimationTracks.Empty();
 }
 #endif
 
@@ -252,87 +255,68 @@ FFrameRate UAnimationSequencerDataModel::GetFrameRate() const
 
 const TArray<FBoneAnimationTrack>& UAnimationSequencerDataModel::GetBoneAnimationTracks() const
 {
-	return LegacyBoneAnimationTracks;
+	static TArray<FBoneAnimationTrack> TempTracks;
+	return TempTracks;
 }
 
 const FBoneAnimationTrack& UAnimationSequencerDataModel::GetBoneTrackByIndex(int32 TrackIndex) const
 {
-	checkf(LegacyBoneAnimationTracks.IsValidIndex(TrackIndex), TEXT("Unable to find animation track by index"));
-	return LegacyBoneAnimationTracks[TrackIndex];
+	static FBoneAnimationTrack TempTrack;
+	return TempTrack;
 }
 
 const FBoneAnimationTrack& UAnimationSequencerDataModel::GetBoneTrackByName(FName TrackName) const
 {
-	const FBoneAnimationTrack* TrackPtr = LegacyBoneAnimationTracks.FindByPredicate([TrackName](const FBoneAnimationTrack& Track)
-	{
-		return Track.Name == TrackName;
-	});
-
-	checkf(TrackPtr != nullptr, TEXT("Unable to find animation track by name"));
-
-	return *TrackPtr;
+	static FBoneAnimationTrack TempTrack;
+	return TempTrack;
 }
 
 const FBoneAnimationTrack* UAnimationSequencerDataModel::FindBoneTrackByName(FName Name) const
 {
-	return LegacyBoneAnimationTracks.FindByPredicate([Name](const FBoneAnimationTrack& Track)
-	{
-		return Track.Name == Name;
-	});
-}
-
-FBoneAnimationTrack* UAnimationSequencerDataModel::FindMutableBoneTrackByName(FName Name)
-{
-	return LegacyBoneAnimationTracks.FindByPredicate([&Name](const FBoneAnimationTrack& Track)
-	{
-		return Track.Name == Name;
-	});
+	return nullptr;
 }
 
 const FBoneAnimationTrack* UAnimationSequencerDataModel::FindBoneTrackByIndex(int32 BoneIndex) const
 {
-	const FBoneAnimationTrack* TrackPtr = LegacyBoneAnimationTracks.FindByPredicate([BoneIndex](const FBoneAnimationTrack& Track)
-	{
-		return Track.BoneTreeIndex == BoneIndex;
-	});
-
-	return TrackPtr;
+	return nullptr;
 }
 
 int32 UAnimationSequencerDataModel::GetBoneTrackIndex(const FBoneAnimationTrack& Track) const
 {
-	return LegacyBoneAnimationTracks.IndexOfByPredicate([&Track](const FBoneAnimationTrack& SearchTrack)
-	{
-		return SearchTrack.Name == Track.Name;
-	});
+	return INDEX_NONE;
 }
 
 int32 UAnimationSequencerDataModel::GetBoneTrackIndexByName(FName TrackName) const
 {
-	if (const FBoneAnimationTrack* TrackPtr = FindBoneTrackByName(TrackName))
-	{
-		return GetBoneTrackIndex(*TrackPtr);
-	}
-
 	return INDEX_NONE;
 }
 
 bool UAnimationSequencerDataModel::IsValidBoneTrackIndex(int32 TrackIndex) const
 {
-	return LegacyBoneAnimationTracks.IsValidIndex(TrackIndex);
+	return false;
 }
 
 int32 UAnimationSequencerDataModel::GetNumBoneTracks() const
 {
-	return LegacyBoneAnimationTracks.Num();
+	ValidateSequencerData();
+		
+	if(const UMovieSceneControlRigParameterSection* Section = GetFKControlRigSection())
+	{
+		return Section->GetTransformParameterNamesAndCurves().Num();
+	}
+
+	return 0;
 }
 
 void UAnimationSequencerDataModel::GetBoneTrackNames(TArray<FName>& OutNames) const
 {
-	Algo::Transform(LegacyBoneAnimationTracks, OutNames, [](const FBoneAnimationTrack& Track)
+	if(const UMovieSceneControlRigParameterSection* Section = GetFKControlRigSection())
 	{
-		return Track.Name; 
-	});
+		for (const FTransformParameterNameAndCurves& TransformParameter : Section->GetTransformParameterNamesAndCurves())
+	{
+			OutNames.Add(UFKControlRig::GetControlTargetName(TransformParameter.ParameterName, ERigElementType::Bone));
+		}
+	}
 }
 
 const FAnimationCurveData& UAnimationSequencerDataModel::GetCurveData() const
@@ -437,6 +421,70 @@ const FRichCurve* UAnimationSequencerDataModel::FindRichCurve(const FAnimationCu
 	return RichCurve;
 }
 
+bool UAnimationSequencerDataModel::IsValidBoneTrackName(const FName& TrackName) const
+{
+	ValidateSequencerData();
+
+	if (const UMovieSceneControlRigParameterSection* Section = GetFKControlRigSection())
+	{
+		const FName ControlName = UFKControlRig::GetControlName(TrackName, ERigElementType::Bone);
+		return Section->GetTransformParameterNamesAndCurves().ContainsByPredicate([ControlName](const FTransformParameterNameAndCurves& Curve) { return Curve.ParameterName == ControlName; });
+	}
+	
+	return false;
+}
+
+FTransform UAnimationSequencerDataModel::GetBoneTrackTransform(FName TrackName, const FFrameNumber& FrameNumber) const
+{
+	const TArray<FFrameNumber> FrameNumbers = { FrameNumber };
+	TArray<FTransform> Transforms;	
+	GenerateTransformKeysForControl(TrackName, FrameNumbers, Transforms);
+	return Transforms.Num() ? Transforms[0] : FTransform::Identity;
+}
+
+void UAnimationSequencerDataModel::GetBoneTrackTransforms(FName TrackName, const TArray<FFrameNumber>& FrameNumbers, TArray<FTransform>& OutTransforms) const
+{
+	GenerateTransformKeysForControl(TrackName, FrameNumbers, OutTransforms);
+}
+
+void UAnimationSequencerDataModel::GetBoneTrackTransforms(FName TrackName, TArray<FTransform>& OutTransforms) const
+{
+	IterateTransformControlCurve(TrackName, [&OutTransforms](const FTransform& Transform, const FFrameNumber& FrameNumber) -> void
+	{
+		OutTransforms.Add(Transform);
+	});
+}
+
+void UAnimationSequencerDataModel::GetBoneTracksTransform(const TArray<FName>& TrackNames, const FFrameNumber& FrameNumber, TArray<FTransform>& OutTransforms) const
+{
+	const TArray<FFrameNumber> FrameNumbers = { FrameNumber };
+	for (int32 EntryIndex = 0; EntryIndex < TrackNames.Num(); ++EntryIndex)
+	{
+		GenerateTransformKeysForControl(TrackNames[EntryIndex], FrameNumbers, OutTransforms);
+	}
+}
+
+FTransform UAnimationSequencerDataModel::EvaluateBoneTrackTransform(FName TrackName, const FFrameTime& FrameTime, const EAnimInterpolationType& Interpolation) const
+{
+	const float Alpha = Interpolation == EAnimInterpolationType::Step ? FMath::RoundToFloat(FrameTime.GetSubFrame()) : FrameTime.GetSubFrame();
+
+	if (FMath::IsNearlyEqual(Alpha, 1.0f))
+	{
+		return GetBoneTrackTransform(TrackName, FrameTime.CeilToFrame());
+	}
+	else if (FMath::IsNearlyZero(Alpha))
+	{
+		return GetBoneTrackTransform(TrackName, FrameTime.FloorToFrame());
+	}
+	
+	const FTransform From = GetBoneTrackTransform(TrackName, FrameTime.FloorToFrame());
+	const FTransform To = GetBoneTrackTransform(TrackName, FrameTime.CeilToFrame());
+
+	FTransform Blend;
+	Blend.Blend(From, To, Alpha);
+	return Blend;
+}
+
 const FAnimCurveBase& UAnimationSequencerDataModel::GetCurve(const FAnimationCurveIdentifier& CurveIdentifier) const
 {
 	const FAnimCurveBase* CurvePtr = FindCurve(CurveIdentifier);
@@ -688,10 +736,6 @@ void UAnimationSequencerDataModel::OnNotify(const EAnimDataModelNotifyType& Noti
 		const TArray<EAnimDataModelNotifyType> BonesNotifyTypes = {EAnimDataModelNotifyType::TrackAdded, EAnimDataModelNotifyType::TrackChanged, EAnimDataModelNotifyType::TrackRemoved, EAnimDataModelNotifyType::Populated, EAnimDataModelNotifyType::Reset };
 		if(Collector.Contains(BonesNotifyTypes))
 		{
-			if(!ValidationMode)
-			{
-				GenerateLegacyBoneData();
-			}
 			RefreshControlsAndProxy();
 			ResetCachedGUID();
 		}
@@ -815,97 +859,6 @@ void UAnimationSequencerDataModel::GenerateLegacyCurveData()
 	}
 }
 
-void UAnimationSequencerDataModel::GenerateLegacyBoneData()
-{
-	QUICK_SCOPE_CYCLE_COUNTER(STAT_GenerateLegacyBoneData);
-	// Reset current track-data
-	LegacyBoneAnimationTracks.Reset();
-
-	if(const USkeleton* TargetSkeleton = GetSkeleton())
-	{
-		const FReferenceSkeleton& ReferenceSkeleton = TargetSkeleton->GetReferenceSkeleton();
-		ValidateSequencerData();
-		
-		if(const UMovieSceneControlRigParameterSection* Section = GetFKControlRigSection())
-		{
-			if (const UControlRig* ControlRig = Section->GetControlRig())
-			{
-				if (const URigHierarchy* Hierarchy = ControlRig->GetHierarchy())
-				{
-					TArray<FTransform> Transforms;
-					TArray<FFrameNumber> FrameNumbers;
-
-					const TArray<FTransformParameterNameAndCurves>& TransformCurves = Section->GetTransformParameterNamesAndCurves();
-					LegacyBoneAnimationTracks.SetNumZeroed(TransformCurves.Num());
-					ParallelFor(TransformCurves.Num(), [this, &TransformCurves, Hierarchy, &ReferenceSkeleton](int32 CurveIndex)
-					{
-						const FTransformParameterNameAndCurves& TransformParameterCurve = TransformCurves[CurveIndex];
-
-						const FName TargetBoneName = UFKControlRig::GetControlTargetName(TransformParameterCurve.ParameterName, ERigElementType::Bone);
-						const FRigElementKey BoneElementKey(TargetBoneName, ERigElementType::Bone);
-						if (Hierarchy->Contains(BoneElementKey))
-						{							
-							// Only populate the track if any curve keys were set
-							const bool bContainsAnyKeys = [&TransformParameterCurve]()
-							{
-								for (int32 ChannelIndex = 0; ChannelIndex < 3; ++ChannelIndex)
-								{
-									if (TransformParameterCurve.Translation[ChannelIndex].HasAnyData() || TransformParameterCurve.Rotation[ChannelIndex].HasAnyData() || TransformParameterCurve.Scale[ChannelIndex].HasAnyData())
-									{
-										return true;
-									}
-								}
-
-								return false;
-							}();
-							
-							FBoneAnimationTrack& BoneTrack = LegacyBoneAnimationTracks[CurveIndex];
-							BoneTrack.Name = TargetBoneName;
-							BoneTrack.BoneTreeIndex = ReferenceSkeleton.FindBoneIndex(BoneTrack.Name);
-
-							if (bContainsAnyKeys)
-							{								
-								const int32 NumKeys = GetNumberOfKeys();
-								BoneTrack.InternalTrackData.PosKeys.SetNumUninitialized(NumKeys);
-								BoneTrack.InternalTrackData.RotKeys.SetNumUninitialized(NumKeys);
-								BoneTrack.InternalTrackData.ScaleKeys.SetNumUninitialized(NumKeys);
-
-								FVector3f EulerAngles;
-								for (int32 FrameIndex = 0; FrameIndex < NumKeys; ++FrameIndex)
-								{
-									for (int32 ChannelIndex = 0; ChannelIndex < 3; ++ChannelIndex)
-									{
-										BoneTrack.InternalTrackData.PosKeys[FrameIndex][ChannelIndex] = TransformParameterCurve.Translation[ChannelIndex].GetValues().Num() == 0 ? TransformParameterCurve.Translation[ChannelIndex].GetDefault().GetValue() : TransformParameterCurve.Translation[ChannelIndex].GetValues()[FrameIndex].Value;
-										EulerAngles[ChannelIndex] = TransformParameterCurve.Rotation[ChannelIndex].GetValues().Num() == 0 ? TransformParameterCurve.Rotation[ChannelIndex].GetDefault().GetValue() : TransformParameterCurve.Rotation[ChannelIndex].GetValues()[FrameIndex].Value;
-										BoneTrack.InternalTrackData.ScaleKeys[FrameIndex][ChannelIndex] = TransformParameterCurve.Scale[ChannelIndex].GetValues().Num() == 0 ? TransformParameterCurve.Scale[ChannelIndex].GetDefault().GetValue() : TransformParameterCurve.Scale[ChannelIndex].GetValues()[FrameIndex].Value;
-									}
-									
-									BoneTrack.InternalTrackData.RotKeys[FrameIndex] = FQuat4f::MakeFromEuler(EulerAngles);
-								}								
-							}
-						}
-					});
-
-					// Remove tracks for bones not present in the new hierarchy
-					LegacyBoneAnimationTracks.RemoveAll([](const FBoneAnimationTrack& Track) { return Track.Name.IsNone() && Track.InternalTrackData.PosKeys.Num() == 0; } );
-				}
-				else
-				{						
-					IAnimationDataController::ReportObjectErrorf(this, LOCTEXT("UnableToFindRigHierarchy", "Unable to retrieve RigHierarchy for ControlRig ({0})"), FText::FromString(ControlRig->GetPathName()));	      
-				}
-			}
-			else
-			{				
-				IAnimationDataController::ReportObjectErrorf(this, LOCTEXT("UnableToFindControlRig", "Unable to retrieve ControlRig for Model ({0})"), FText::FromString(GetPathName()));
-			}
-		}					
-	}
-	else
-	{
-		IAnimationDataController::ReportObjectErrorf(this, LOCTEXT("UnableToFindSkeleton", "Unable to retrieve target USkeleton for Animation Asset ({0})"), FText::FromString(GetOuter()->GetPathName()));	
-	}
-}
-
 void UAnimationSequencerDataModel::ValidateData() const
 {		
 	ValidateSequencerData();
@@ -990,82 +943,7 @@ void UAnimationSequencerDataModel::ValidateLegacyAgainstControlRigData() const
 	const UAnimSequence* OuterSequence = GetAnimationSequence();
 	if (const USkeleton* Skeleton = OuterSequence->GetSkeleton())
 	{
-		const FReferenceSkeleton& ReferenceSkeleton = Skeleton->GetReferenceSkeleton();
-
-		TArray<FTransform> Transforms;
-		TArray<FFrameNumber> FrameNumbers;
-		
-		for (const FBoneAnimationTrack& Track : LegacyBoneAnimationTracks)
-		{
-			const FName ExpectedBoneName = ReferenceSkeleton.GetBoneName(Track.BoneTreeIndex);
-			
-			const FRigElementKey BoneKey(ExpectedBoneName, ERigElementType::Bone);
-			const FRigBoneElement* BoneElement = Hierarchy->Find<FRigBoneElement>(BoneKey);
-			if(!BoneElement)
-			{
-				IAnimationDataController::ReportObjectErrorf(this, LOCTEXT("BoneElementNotFound", "Unable to find FRigBoneElement in RigHierarchy for Bone with name: {0}"), FText::FromString(ExpectedBoneName.ToString()));
-			}
-
-			const FRigElementKey BoneControlKey(UFKControlRig::GetControlName(ExpectedBoneName, ERigElementType::Bone), ERigElementType::Control);
-			const FRigControlElement* BoneControlElement = Hierarchy->Find<FRigControlElement>(BoneControlKey);
-
-			if (!BoneControlElement)
-			{
-				IAnimationDataController::ReportObjectErrorf(this, LOCTEXT("ControlElementNotFound", "Unable to find FRigControlElement in RigHierarchy for Bone with name: {0}"), FText::FromString(ExpectedBoneName.ToString()));
-			}		
-		
-			const FTransformParameterNameAndCurves* BoneCurveParameter = Section->GetTransformParameterNamesAndCurves().FindByPredicate([BoneControlKey](const FTransformParameterNameAndCurves& ParameterPair)
-			{
-				return ParameterPair.ParameterName == BoneControlKey.Name;
-			});
-			if (!(BoneCurveParameter || Track.InternalTrackData.PosKeys.Num() == 0))
-			{
-				IAnimationDataController::ReportObjectErrorf(this, LOCTEXT("ControlCurveNotFound", "Unable to find FTransformParameterNameAndCurves in RigHierarchy for Bone Control with name: {0}"), FText::FromName(BoneControlKey.Name));
-			}
-
-			GenerateTransformKeysForControl(ExpectedBoneName, Transforms, FrameNumbers);
-
-			const int32 NumExpectedKeys = Track.InternalTrackData.PosKeys.Num();
-			if (NumExpectedKeys != Transforms.Num())
-			{
-				IAnimationDataController::ReportObjectErrorf(this, LOCTEXT("UnexpectedNumberOfControlKeys",	"Unexpected number of Bone Control Curve keys for {0}, expected {1} but found {2}"), FText::FromName(ExpectedBoneName), FText::AsNumber(NumExpectedKeys), FText::AsNumber(Transforms.Num()));
-			}			
-
-			if (NumExpectedKeys == Transforms.Num())
-			{
-				for (int32 KeyIndex = 0; KeyIndex < NumExpectedKeys; ++KeyIndex)
-				{
-					const FTransform& TransformKey = Transforms[KeyIndex];
-
-					checkf(TransformKey.GetLocation().Equals(FVector(Track.InternalTrackData.PosKeys[KeyIndex])), TEXT("Unexpected positional key (%i) for bone %s, expected %s but found %s"), KeyIndex,
-						*ExpectedBoneName.ToString(),* Track.InternalTrackData.PosKeys[KeyIndex].ToCompactString(), *TransformKey.GetLocation().ToCompactString());
-
-					if(!TransformKey.GetLocation().Equals(FVector(Track.InternalTrackData.PosKeys[KeyIndex])))
-					{
-						IAnimationDataController::ReportObjectErrorf(this, LOCTEXT("UnexpectedPositionalKey", "Unexpected positional key ({0}) for bone {1}, expected {2} but found {3}"),							
-							FText::AsNumber(KeyIndex), FText::FromName(ExpectedBoneName), FText::FromString(Track.InternalTrackData.PosKeys[KeyIndex].ToCompactString()), FText::FromString(TransformKey.GetLocation().ToCompactString()));
-					}
-
-					const FQuat LegacyRotation = FQuat(Track.InternalTrackData.RotKeys[KeyIndex]).GetNormalized();
-					const double RotationDeltaDegrees = FMath::RadiansToDegrees(TransformKey.GetRotation().AngularDistance(LegacyRotation));
-			
-					if(RotationDeltaDegrees > 0.5f)
-					{
-						IAnimationDataController::ReportObjectErrorf(this, LOCTEXT("UnexpectedRotationalKey", "Unexpected rotational key ({0} degrees delta) ({1}) for bone {2}, expected {3} but found {4}"), FText::AsNumber(RotationDeltaDegrees), FText::AsNumber(KeyIndex), FText::FromName(ExpectedBoneName), FText::FromString(Track.InternalTrackData.RotKeys[KeyIndex].ToString()), FText::FromString(TransformKey.GetRotation().ToString()));
-					}
-
-					if (!TransformKey.GetScale3D().Equals(FVector(Track.InternalTrackData.ScaleKeys[KeyIndex])))
-					{
-						IAnimationDataController::ReportObjectErrorf(this, LOCTEXT("UnexpectedScalingKey", "Unexpected scaling key ({0}) for bone {1}, expected {2} but found {3}"), FText::AsNumber(KeyIndex),
-							FText::FromName(ExpectedBoneName), FText::FromString(Track.InternalTrackData.ScaleKeys[KeyIndex].ToCompactString()), FText::FromString(TransformKey.GetScale3D().ToCompactString()));
-					}
-				}
-			
-				Transforms.Reset();
-				FrameNumbers.Reset();
-			}
-		}
-		
+		const FReferenceSkeleton& ReferenceSkeleton = Skeleton->GetReferenceSkeleton();	
 		// Validate curve data against controls
 		for (const FFloatCurve& FloatCurve : LegacyCurveData.FloatCurves)
 		{
@@ -1117,7 +995,7 @@ void UAnimationSequencerDataModel::ValidateLegacyAgainstControlRigData() const
 	}	
 }
 
-void UAnimationSequencerDataModel::IterateTransformControlCurve(const FName& BoneName, TFunction<void(const FTransform&, const FFrameNumber&)> IterationFunction) const
+void UAnimationSequencerDataModel::IterateTransformControlCurve(const FName& BoneName, TFunction<void(const FTransform&, const FFrameNumber&)> IterationFunction, const TArray<FFrameNumber>* InFrameNumbers /*= nullptr*/) const
 {
 	ValidateSequencerData();
 	ValidateControlRigData();
@@ -1136,31 +1014,52 @@ void UAnimationSequencerDataModel::IterateTransformControlCurve(const FName& Bon
 		{
 			const FTransformParameterNameAndCurves& ControlCurve = *ControlCurvePtr;
 
-			FTransform Transform;
-			FVector3f Location;
-			FVector3f EulerAngles;
-			FVector3f Scale;
+			FVector3f Location(0.f);
+			FVector3f EulerAngles(0.f);
+			FVector3f Scale = FVector3f::OneVector;
 
-			for (int32 KeyIndex = 0; KeyIndex < GetNumberOfKeys(); ++KeyIndex)
+			// Check whether or not any data is contained
+			bool bContainsData = false;
+			bool bContainsKeys = false;
+			for (int32 ChannelIndex = 0; ChannelIndex < 3; ++ChannelIndex)
 			{
-				const FFrameNumber Frame(KeyIndex);
-				for (int32 ChannelIndex = 0; ChannelIndex < 3; ++ChannelIndex)
+				bContainsData |= ControlCurve.Translation[ChannelIndex].HasAnyData();
+				bContainsKeys |= ControlCurve.Translation[ChannelIndex].GetNumKeys() != 0;
+				bContainsData |= ControlCurve.Rotation[ChannelIndex].HasAnyData();		
+				bContainsKeys |= ControlCurve.Rotation[ChannelIndex].GetNumKeys() != 0;		
+				bContainsData |= ControlCurve.Scale[ChannelIndex].HasAnyData();
+				bContainsKeys |= ControlCurve.Scale[ChannelIndex].GetNumKeys() != 0;
+			}
+
+			if (bContainsData)
+			{
+				const int32 NumberOfKeysToIterate = InFrameNumbers != nullptr ? InFrameNumbers->Num() : GetNumberOfKeys();
+				for (int32 KeyIndex = 0; KeyIndex < NumberOfKeysToIterate; ++KeyIndex)
 				{
-					ControlCurve.Translation[ChannelIndex].Evaluate(Frame, Location[ChannelIndex]);
-					ControlCurve.Rotation[ChannelIndex].Evaluate(Frame, EulerAngles[ChannelIndex]);
-					ControlCurve.Scale[ChannelIndex].Evaluate(Frame, Scale[ChannelIndex]);
-				}
+					const FFrameNumber Frame = InFrameNumbers != nullptr ? (*InFrameNumbers)[KeyIndex] : FFrameNumber(KeyIndex);
+					for (int32 ChannelIndex = 0; ChannelIndex < 3; ++ChannelIndex)
+					{
+						ControlCurve.Translation[ChannelIndex].Evaluate(Frame, Location[ChannelIndex]);
+						ControlCurve.Rotation[ChannelIndex].Evaluate(Frame, EulerAngles[ChannelIndex]);
+						ControlCurve.Scale[ChannelIndex].Evaluate(Frame, Scale[ChannelIndex]);
+					}
 
-				Transform.SetLocation(FVector(Location));
-				Transform.SetRotation(FQuat::MakeFromEuler(FVector(EulerAngles)));
-				Transform.SetScale3D(FVector(Scale));
+					FTransform Transform;
+					Transform.SetLocation(FVector(Location));
+					Transform.SetRotation(FQuat::MakeFromEuler(FVector(EulerAngles)));
+					Transform.SetScale3D(FVector(Scale));
 
-				Transform.NormalizeRotation();
+					Transform.NormalizeRotation();
 
-				IterationFunction(Transform, Frame);
+					IterationFunction(Transform, Frame);
+				}
 			}
 		}
 	}
+	 else
+	 {
+		 ensure(false);
+	 }
 }
 
 void UAnimationSequencerDataModel::GenerateTransformKeysForControl(const FName& BoneName, TArray<FTransform>& InOutTransforms, TArray<FFrameNumber>& InOutFrameNumbers) const
@@ -1172,6 +1071,14 @@ void UAnimationSequencerDataModel::GenerateTransformKeysForControl(const FName&
 	});
 }
 
+void UAnimationSequencerDataModel::GenerateTransformKeysForControl(const FName& BoneName, const TArray<FFrameNumber>& FrameNumbers, TArray<FTransform>& InOutTransforms) const
+{
+	IterateTransformControlCurve(BoneName, [&InOutTransforms](const FTransform& Transform, const FFrameNumber& FrameNumber) -> void
+	{
+		InOutTransforms.Add(Transform);
+	}, &FrameNumbers);
+}
+
 UMovieScene* UAnimationSequencerDataModel::GetMovieScene() const
 {
 	return MovieScene;
@@ -1197,7 +1104,6 @@ void UAnimationSequencerDataModel::GeneratePoseData(UControlRig* ControlRig, FAn
 			FCompactPose& RigPose = InOutPoseData.GetPose();
 			RigPose.ResetToRefPose();
 			const FBoneContainer& RequiredBones = RigPose.GetBoneContainer();
-
 			FBlendedCurve& Curve = InOutPoseData.GetCurve();
 			UE::Anim::Retargeting::FRetargetingScope RetargetingScope(RigPose, EvaluationContext);
 			
@@ -1206,51 +1112,44 @@ void UAnimationSequencerDataModel::GeneratePoseData(UControlRig* ControlRig, FAn
 				QUICK_SCOPE_CYCLE_COUNTER(STAT_GetMappings);
 				const FReferenceSkeleton& RefSkeleton = RequiredBones.GetReferenceSkeleton();
 
-				RigHierarchy->ForEach<FRigControlElement>([this, &RefSkeleton, &Curve, &RequiredBones, &RigHierarchy, &RetargetingScope, &RigPose, bValidCurve = Curve.IsValid(), SmartNameContainer=RequiredBones.GetSkeletonAsset()->GetSmartNameContainer(USkeleton::AnimCurveMappingName)](const FRigControlElement* ControlElement) -> bool
+				RigHierarchy->ForEach<FRigBoneElement>([&RefSkeleton, &RequiredBones, &RigHierarchy, &RetargetingScope, &RigPose](const FRigBoneElement* BoneElement) -> bool
 				{
-					if (ControlElement->Settings.ControlType == ERigControlType::EulerTransform)
+					const FName& BoneName = BoneElement->GetName();
+					const int32 BoneIndex = RefSkeleton.FindBoneIndex(BoneName);
+					if (BoneIndex != INDEX_NONE)
 					{
-						const FName ExpectedBoneName = UFKControlRig::GetControlTargetName(ControlElement->GetName(), ERigElementType::Bone);
-						const int32 BoneIndex = RefSkeleton.FindBoneIndex(ExpectedBoneName);
-						if (BoneIndex != INDEX_NONE)
+						const int32 SkeletonBoneIndex = RequiredBones.GetSkeletonAsset()->GetReferenceSkeleton().FindBoneIndex(BoneName);
+						const FCompactPoseBoneIndex CompactPoseBoneIndex = RequiredBones.GetCompactPoseIndexFromSkeletonIndex(SkeletonBoneIndex);
+						if (CompactPoseBoneIndex != INDEX_NONE)
 						{
-							const FName& BoneName = ExpectedBoneName;
-							const FRigElementKey Key(BoneName, ERigElementType::Bone);
-
-							const bool bMatchingLegacyBone = !ValidationMode || LegacyBoneAnimationTracks.ContainsByPredicate([BoneName](const FBoneAnimationTrack& Track)
-							{
-								return Track.Name == BoneName;
-							});
-							ensureMsgf(bMatchingLegacyBone, TEXT("Non-matching bone vs legacy data %s"), *BoneName.ToString());
-						
-							const int32 SkeletonBoneIndex = RequiredBones.GetSkeletonAsset()->GetReferenceSkeleton().FindBoneIndex(BoneName);
-							const FCompactPoseBoneIndex CompactPoseBoneIndex = RequiredBones.GetCompactPoseIndexFromSkeletonIndex(SkeletonBoneIndex);
-							if (CompactPoseBoneIndex != INDEX_NONE)
-							{
-								RetargetingScope.AddTrackedBone(CompactPoseBoneIndex, SkeletonBoneIndex);
-				        
-								// Retrieve evaluated bone transform from Hierarchy
-								RigPose[CompactPoseBoneIndex] = RigHierarchy->GetLocalTransform(Key);
-							}
+							RetargetingScope.AddTrackedBone(CompactPoseBoneIndex, SkeletonBoneIndex);
+							// Retrieve evaluated bone transform from Hierarchy
+							RigPose[CompactPoseBoneIndex] = RigHierarchy->GetLocalTransform(BoneElement->GetKey());
 						}
 					}
-					else if (ControlElement->Settings.ControlType == ERigControlType::Float && bValidCurve)
+
+					return true;
+				});
+
+				if (Curve.IsValid())
+				{
+					const FSmartNameMapping* SmartNameMapping = RequiredBones.GetSkeletonAsset()->GetSmartNameContainer(USkeleton::AnimCurveMappingName);
+					RigHierarchy->ForEach<FRigCurveElement>([this, &RigHierarchy, &Curve, SmartNameMapping](const FRigCurveElement* CurveElement) -> bool
 					{
-						const FName ExpectedCurveName = UFKControlRig::GetControlTargetName(ControlElement->GetName(), ERigElementType::Curve);
-						const SmartName::UID_Type CurveIndex = SmartNameContainer->FindUID(ExpectedCurveName);
+						const FName& CurveName = CurveElement->GetName();
+						const SmartName::UID_Type CurveIndex = SmartNameMapping->FindUID(CurveName);
 						if (CurveIndex != INDEX_NONE)
 						{
-							const FRigElementKey Key(ExpectedCurveName, ERigElementType::Curve);	
 							if (Curve.IsEnabled(CurveIndex))
 							{
-								Curve.Set(CurveIndex, RigHierarchy->GetCurveValue(Key));
+								Curve.Set(CurveIndex, RigHierarchy->GetCurveValue(CurveElement->GetKey()));
 							}
 							
 						}
-					}
+						return true;
+					});
+				}
 
-					return true;
-				});
 			}			
 
 			{
@@ -1275,13 +1174,17 @@ void UAnimationSequencerDataModel::GeneratePoseData(UControlRig* ControlRig, FAn
 					FTransform Value = TransformCurve.Evaluate(EvaluationContext.SampleFrameRate.AsSeconds(EvaluationContext.SampleTime), 1.f);
 
 					const FCompactPoseBoneIndex BoneIndex(RigPose.GetBoneContainer().GetPoseBoneIndexForBoneName(CurveName));
-					if(ensure(BoneIndex != INDEX_NONE))
+					if(BoneIndex != INDEX_NONE)
 					{
 						const FTransform LocalTransform = RigPose[BoneIndex];
 						RigPose[BoneIndex].SetRotation(LocalTransform.GetRotation() * Value.GetRotation());
 						RigPose[BoneIndex].SetTranslation(LocalTransform.TransformPosition(Value.GetTranslation()));
 						RigPose[BoneIndex].SetScale3D(LocalTransform.GetScale3D() * Value.GetScale3D());
-					}					
+					}
+					else
+					{
+						IAnimationDataController::ReportObjectWarningf(this, LOCTEXT("TransformCurveBoneNotFound", "Failed to find BoneIndex for transform curve %s"), FText::FromName(CurveName));
+					}
 				}
 			}
 
@@ -1346,137 +1249,160 @@ void UAnimationSequencerDataModel::EvaluateTrack(UMovieSceneControlRigParameterT
 		check(ControlRig);
 
 		// Reset to ref-pose
-		ControlRig->GetHierarchy()->ResetPoseToInitial(ERigElementType::Bone);
-
-		const TArray<FScalarParameterNameAndCurve>& ScalarParameters = FKRigSection->GetScalarParameterNamesAndCurves();
-		for (const FScalarParameterNameAndCurve& TypedParameter : ScalarParameters)
+		if (URigHierarchy* RigHierarchy = ControlRig->GetHierarchy())
 		{
-			const FName& Name = TypedParameter.ParameterName;
-			float Value = 0.f;
+			RigHierarchy->ResetPoseToInitial(ERigElementType::Bone);
 
-			const FFrameTime CurveSampleTime = FFrameRate::TransformTime(EvaluationContext.SampleTime, EvaluationContext.SampleFrameRate, TypedParameter.ParameterCurve.GetTickResolution());
-			if(TypedParameter.ParameterCurve.Evaluate(CurveSampleTime, Value))
-			{					
-				const FRigControlElement* ControlElement = ControlRig->FindControl(Name);
-				if (ControlElement && ControlElement->Settings.ControlType == ERigControlType::Float)
-				{
-					ControlRig->SetControlValue<float>(Name, Value, false, EControlRigSetKey::Never, false);
+			const TArray<FScalarParameterNameAndCurve>& ScalarParameters = FKRigSection->GetScalarParameterNamesAndCurves();
+			for (const FScalarParameterNameAndCurve& TypedParameter : ScalarParameters)
+			{
+				const FName& Name = TypedParameter.ParameterName;
+				float Value = 0.f;
+
+				const FFrameTime CurveSampleTime = FFrameRate::TransformTime(EvaluationContext.SampleTime, EvaluationContext.SampleFrameRate, TypedParameter.ParameterCurve.GetTickResolution());
+				if(TypedParameter.ParameterCurve.Evaluate(CurveSampleTime, Value))
+				{					
+					FRigControlElement* ControlElement = ControlRig->FindControl(Name);
+					if (ControlElement && ControlElement->Settings.ControlType == ERigControlType::Float)
+					{
+						RigHierarchy->SetControlValue(ControlElement, FRigControlValue::Make<float>(Value), ERigControlValueType::Current, false, false, false, false);
+					}
 				}
 			}
-		}
 
-		const TArray<FTransformParameterNameAndCurves>& TransformParameters = FKRigSection->GetTransformParameterNamesAndCurves();
-		if (TransformParameters.Num())
-		{
-			QUICK_SCOPE_CYCLE_COUNTER(STAT_EvaluateTransformParameters);
-			
-			struct FEvaluationInfo
+			const TArray<FTransformParameterNameAndCurves>& TransformParameters = FKRigSection->GetTransformParameterNamesAndCurves();
+			if (TransformParameters.Num())
 			{
-				float Interp = 0.f;
-				int32 Index1 = INDEX_NONE, Index2 = INDEX_NONE;
-			};
+				QUICK_SCOPE_CYCLE_COUNTER(STAT_EvaluateTransformParameters);
+			
+				struct FEvaluationInfo
+				{
+					float Interp = 0.f;
+					int32 Index1 = INDEX_NONE, Index2 = INDEX_NONE;
+				};
 
-			TMovieSceneCurveChannelImpl<FMovieSceneFloatChannel>::FTimeEvaluationCache FromFrameTimeEvaluationCache;
-			TMovieSceneCurveChannelImpl<FMovieSceneFloatChannel>::FTimeEvaluationCache ToFrameTimeEvaluationCache;
-			for (const FTransformParameterNameAndCurves& TypedParameter : TransformParameters)
-			{
-				const FName& Name = TypedParameter.ParameterName;
-				const FRigControlElement* ControlElement = ControlRig->FindControl(Name);
-				if (ControlElement && ControlElement->Settings.ControlType == ERigControlType::EulerTransform)
+				TMovieSceneCurveChannelImpl<FMovieSceneFloatChannel>::FTimeEvaluationCache FromFrameTimeEvaluationCache;
+				TMovieSceneCurveChannelImpl<FMovieSceneFloatChannel>::FTimeEvaluationCache ToFrameTimeEvaluationCache;
+
+				const int32 NumberOfKeys = GetNumberOfKeys();
+			
+				for (const FTransformParameterNameAndCurves& TypedParameter : TransformParameters)
 				{
-					FEulerTransform EulerTransform;
-					
-					const double Alpha = BoneSampleTime.GetSubFrame();			
-					auto EvaluateToTransform = [&TypedParameter](const FFrameNumber& Frame, FTransform& InOutTransform, TMovieSceneCurveChannelImpl<FMovieSceneFloatChannel>::FTimeEvaluationCache& Cache)
+					const FName& Name = TypedParameter.ParameterName;
+					FRigControlElement* ControlElement = ControlRig->FindControl(Name);
+					if (ControlElement && ControlElement->Settings.ControlType == ERigControlType::EulerTransform)
 					{
-						auto EvaluateValue = [Frame, &Cache](const auto& Channel, auto& Target)
+						FEulerTransform EulerTransform(FEulerTransform::Identity);
+					
+						const double Alpha = BoneSampleTime.GetSubFrame();
+						auto EvaluateToTransform = [&TypedParameter](const FFrameNumber& Frame, FTransform& InOutTransform, TMovieSceneCurveChannelImpl<FMovieSceneFloatChannel>::FTimeEvaluationCache* Cache)
 						{
-							auto Value = 0.f;
-							TMovieSceneCurveChannelImpl<FMovieSceneFloatChannel>::EvaluateWithCache(&Channel, &Cache, Frame, Value);
-							Target = Value;
+							auto EvaluateValue = [Frame, Cache](const auto& Channel, auto& Target)
+							{
+								if (Channel.GetDefault().IsSet())
+								{
+									Target = Channel.GetDefault().GetValue();
+								}
+								else
+								{
+									auto Value = (float)Target;
+									TMovieSceneCurveChannelImpl<FMovieSceneFloatChannel>::EvaluateWithCache(&Channel, Cache, Frame, Value);
+									Target = Value;
+								}
+							};
+
+							auto EvaluateVector = [EvaluateValue](const auto& VectorChannels, auto& TargetVector)
+							{
+								EvaluateValue(VectorChannels[0], TargetVector[0]);
+								EvaluateValue(VectorChannels[1], TargetVector[1]);
+								EvaluateValue(VectorChannels[2], TargetVector[2]);
+							};
+
+							FVector Location(FVector::ZeroVector), Scale(FVector::OneVector);
+							EvaluateVector(TypedParameter.Translation, Location);
+							InOutTransform.SetTranslation(Location);
+							EvaluateVector(TypedParameter.Scale, Scale);
+							InOutTransform.SetScale3D(Scale);
+
+							FRotator Rotator;
+							EvaluateValue(TypedParameter.Rotation[0], Rotator.Roll);
+							EvaluateValue(TypedParameter.Rotation[1], Rotator.Pitch);
+							EvaluateValue(TypedParameter.Rotation[2], Rotator.Yaw);
+						
+							InOutTransform.SetRotation(Rotator.Quaternion());
 						};
 
-						auto EvaluateVector = [EvaluateValue](const auto& VectorChannels, auto& TargetVector)
+						auto ExtractTransform = [&TypedParameter, NumberOfKeys](const FFrameNumber& Frame, FEulerTransform& InOutEulerTransform, TMovieSceneCurveChannelImpl<FMovieSceneFloatChannel>::FTimeEvaluationCache* Cache)
 						{
-							EvaluateValue(VectorChannels[0], TargetVector[0]);
-							EvaluateValue(VectorChannels[1], TargetVector[1]);
-							EvaluateValue(VectorChannels[2], TargetVector[2]);
-						};
+							auto ExtractValue = [&TypedParameter, Frame, Cache, NumberOfKeys](const auto& Channel, auto& Target)
+							{
+								if (Channel.HasAnyData())
+								{
+									const int32 NumValues = Channel.GetValues().Num();
+									// No keys, but has data so DefaultValue is set
+									if (NumValues == 0)
+									{
+										Target = Channel.GetDefault().GetValue();
+									}
+									// Uniform keys
+									else if (NumValues == NumberOfKeys)
+									{
+										Target = Channel.GetValues()[Frame.Value].Value;
+									}
+									// Non-uniform keys
+									else
+									{
+										auto Value = (float)Target;
+										TMovieSceneCurveChannelImpl<FMovieSceneFloatChannel>::EvaluateWithCache(&Channel, Cache, Frame, Value);
+										Target = Value;
+									}
+								}
+							};
 
-						FVector Location, Scale;
-						EvaluateVector(TypedParameter.Translation, Location);
-						InOutTransform.SetTranslation(Location);
-						EvaluateVector(TypedParameter.Scale, Scale);
-						InOutTransform.SetScale3D(Scale);
+							auto ExtractVector = [ExtractValue](const auto& VectorChannels, auto& TargetVector)
+							{
+								ExtractValue(VectorChannels[0], TargetVector[0]);
+								ExtractValue(VectorChannels[1], TargetVector[1]);
+								ExtractValue(VectorChannels[2], TargetVector[2]);
+							};
 
-						FRotator Rotator;
-						EvaluateValue(TypedParameter.Rotation[0], Rotator.Roll);
-						EvaluateValue(TypedParameter.Rotation[1], Rotator.Pitch);
-						EvaluateValue(TypedParameter.Rotation[2], Rotator.Yaw);
-						InOutTransform.SetRotation(Rotator.Quaternion());
-					};
+							ExtractVector(TypedParameter.Translation, InOutEulerTransform.Location);
+							ExtractVector(TypedParameter.Scale, InOutEulerTransform.Scale);
 
-					auto ExtractTransform = [&TypedParameter](const FFrameNumber& Frame, FEulerTransform& InOutEulerTransform)
-					{
-						auto ExtractValue = [&TypedParameter, Frame](const auto& Channel, auto& Target)
+							ExtractValue(TypedParameter.Rotation[0], InOutEulerTransform.Rotation.Roll);
+							ExtractValue(TypedParameter.Rotation[1], InOutEulerTransform.Rotation.Pitch);
+							ExtractValue(TypedParameter.Rotation[2], InOutEulerTransform.Rotation.Yaw);
+						};
+						// Assume no interpolation due to uniform keys
+						if (FMath::IsNearlyZero(Alpha))
 						{
-							if (Channel.GetDefault().IsSet())
+							if (EvaluationContext.InterpolationType == EAnimInterpolationType::Linear)
 							{
-								Target = Channel.GetDefault().GetValue();
+								ExtractTransform(BoneSampleTime.FrameNumber, EulerTransform, &FromFrameTimeEvaluationCache);
 							}
-							else
+							else if (EvaluationContext.InterpolationType == EAnimInterpolationType::Step)
 							{
-								Target = Channel.GetValues()[Frame.Value].Value;
+								ExtractTransform(BoneSampleTime.FrameNumber, EulerTransform, &FromFrameTimeEvaluationCache);
 							}
-						};
-
-						auto ExtractVector = [ExtractValue](const auto& VectorChannels, auto& TargetVector)
+						}
+						// Interpolate between two uniform keys
+						else
 						{
-							ExtractValue(VectorChannels[0], TargetVector[0]);
-							ExtractValue(VectorChannels[1], TargetVector[1]);
-							ExtractValue(VectorChannels[2], TargetVector[2]);
-						};
+							const FFrameNumber FromFrame = BoneSampleTime.FloorToFrame();
+							const FFrameNumber ToFrame = BoneSampleTime.CeilToFrame();
 
-						ExtractVector(TypedParameter.Translation, InOutEulerTransform.Location);
-						ExtractVector(TypedParameter.Scale, InOutEulerTransform.Scale);
+							FTransform FromBoneTransform;
+							EvaluateToTransform(FromFrame, FromBoneTransform, &FromFrameTimeEvaluationCache);
+							FTransform ToBoneTransform;
+							EvaluateToTransform(ToFrame, ToBoneTransform, &ToFrameTimeEvaluationCache);
 
-						ExtractValue(TypedParameter.Rotation[0], InOutEulerTransform.Rotation.Roll);
-						ExtractValue(TypedParameter.Rotation[1], InOutEulerTransform.Rotation.Pitch);
-						ExtractValue(TypedParameter.Rotation[2], InOutEulerTransform.Rotation.Yaw);
-					};
-					
-					// Assume no interpolation due to uniform keys
-					if (FMath::IsNearlyZero(Alpha))
-					{
-						if (EvaluationContext.InterpolationType == EAnimInterpolationType::Linear)
-						{
 							FTransform FinalTransform;
-							EvaluateToTransform(BoneSampleTime.FrameNumber, FinalTransform, FromFrameTimeEvaluationCache);
+							FinalTransform.Blend(FromBoneTransform, ToBoneTransform, Alpha);
+						
 							EulerTransform = FEulerTransform(FinalTransform);
 						}
-						else if (EvaluationContext.InterpolationType == EAnimInterpolationType::Step)
-						{
-							ExtractTransform(BoneSampleTime.FrameNumber, EulerTransform);
-						}
+						RigHierarchy->SetControlValue(ControlElement, FRigControlValue::Make<FRigControlValue::FEulerTransform_Float>(EulerTransform), ERigControlValueType::Current, false, false, false, false);
 					}
-					// Interpolate between two uniform keys
-					else
-					{
-						const FFrameNumber FromFrame = BoneSampleTime.FloorToFrame();
-						const FFrameNumber ToFrame = BoneSampleTime.CeilToFrame();
-
-						FTransform FromBoneTransform;
-						EvaluateToTransform(FromFrame, FromBoneTransform, FromFrameTimeEvaluationCache);
-						FTransform ToBoneTransform;
-						EvaluateToTransform(ToFrame, ToBoneTransform, ToFrameTimeEvaluationCache);
-
-						FTransform FinalTransform;
-						FinalTransform.Blend(FromBoneTransform, ToBoneTransform, Alpha);
-						
-						EulerTransform = FEulerTransform(FinalTransform);
-					}
-					
-					ControlRig->SetControlValue<FRigControlValue::FEulerTransform_Float>(Name, EulerTransform, false, EControlRigSetKey::Never, false);
 				}
 			}
 		}
@@ -1560,4 +1486,123 @@ FRichCurve* UAnimationSequencerDataModel::GetMutableRichCurve(const FAnimationCu
 	return RichCurve;
 }
 
+void UAnimationSequencerDataModel::IterateBoneKeys(const FName& BoneName, TFunction<bool(const FVector3f& Pos, const FQuat4f&, const FVector3f, const FFrameNumber&)> IterationFunction) const
+{
+	ValidateSequencerData();
+	ValidateControlRigData();
+
+	const UMovieSceneControlRigParameterSection* Section = GetFKControlRigSection();
+	UControlRig* ControlRig = Section->GetControlRig();
+	const URigHierarchy* Hierarchy = ControlRig->GetHierarchy();
+	
+	const FRigElementKey BoneControlKey(UFKControlRig::GetControlName(BoneName, ERigElementType::Bone), ERigElementType::Control);
+	if (Hierarchy->Contains(BoneControlKey))
+	{
+		if (const FTransformParameterNameAndCurves* ControlCurvePtr = Section->GetTransformParameterNamesAndCurves().FindByPredicate([CurveName = BoneControlKey.Name](const FTransformParameterNameAndCurves& TransformParameter)
+		{
+			return TransformParameter.ParameterName == CurveName;
+		}))
+		{
+			const FTransformParameterNameAndCurves& ControlCurve = *ControlCurvePtr;
+
+			struct FChannelInfo
+			{
+				bool bConstant = true;
+				bool bUniform = true;
+			};			
+			FChannelInfo PosChannels[3];
+			FChannelInfo RotChannels[3];
+			FChannelInfo ScaleChannels[3];
+
+			const int32 NumberOfKeys = GetNumberOfKeys();
+
+			int32 MaxNumberOfKeys = 1;
+			TSet<FFrameNumber> FrameNumbers;
+			for (int32 ChannelIndex = 0; ChannelIndex < 3; ++ChannelIndex)
+			{
+				auto ProcessChannel = [&MaxNumberOfKeys, NumberOfKeys](const FMovieSceneFloatChannel& Channel, FChannelInfo& Info)
+				{
+					if (Channel.HasAnyData())
+					{
+						if (Channel.GetNumKeys() == 0)
+						{
+							Info.bUniform = false;
+						}
+						else
+						{
+							Info.bConstant = false;
+							if (Channel.GetNumKeys() != NumberOfKeys)
+							{
+								Info.bUniform = false;
+							}
+
+							MaxNumberOfKeys = FMath::Max(MaxNumberOfKeys, Channel.GetNumKeys());
+						}
+					}
+				};
+
+				ProcessChannel(ControlCurve.Translation[ChannelIndex], PosChannels[ChannelIndex]);
+				ProcessChannel(ControlCurve.Rotation[ChannelIndex], RotChannels[ChannelIndex]);
+				ProcessChannel(ControlCurve.Scale[ChannelIndex], ScaleChannels[ChannelIndex]);
+			}
+
+
+			const int32 NumberOfKeysToIterate = MaxNumberOfKeys;
+			FVector3f PreviousPos, PreviousScale, PreviousRot;
+			FQuat4f PreviousQuat;
+
+			for (int32 KeyIndex = 0; KeyIndex < NumberOfKeysToIterate; ++KeyIndex)
+			{
+				const FFrameNumber Frame = KeyIndex;
+				
+				for (int32 ChannelIndex = 0; ChannelIndex < 3; ++ChannelIndex)
+				{
+					if (PosChannels[ChannelIndex].bConstant)
+					{
+						if (KeyIndex == 0)
+						{
+							PreviousPos[ChannelIndex] = ControlCurve.Translation[ChannelIndex].GetDefault().GetValue();
+						}
+					}
+					else
+					{
+						ControlCurve.Translation[ChannelIndex].Evaluate(Frame, PreviousPos[ChannelIndex]);
+					}
+					
+					if (RotChannels[ChannelIndex].bConstant)
+					{
+						if (KeyIndex == 0)
+						{
+							PreviousRot[ChannelIndex] = ControlCurve.Rotation[ChannelIndex].GetDefault().GetValue();
+						}
+					}
+					else
+					{
+						ControlCurve.Rotation[ChannelIndex].Evaluate(Frame, PreviousRot[ChannelIndex]);
+					}
+
+					if (ScaleChannels[ChannelIndex].bConstant)
+					{
+						if (KeyIndex == 0)
+						{
+							PreviousScale[ChannelIndex] = ControlCurve.Scale[ChannelIndex].GetDefault().GetValue();
+						}
+					}
+					else
+					{
+						ControlCurve.Scale[ChannelIndex].Evaluate(Frame, PreviousScale[ChannelIndex]);
+					}
+					
+				}
+
+				PreviousQuat = FQuat4f::MakeFromEuler(PreviousRot);
+				if(!IterationFunction(PreviousPos, PreviousQuat, PreviousScale, Frame))
+				{
+					return;
+				}
+			}
+		}
+	}
+}
+
 #undef LOCTEXT_NAMESPACE //"AnimSequencerDataModel"
\ No newline at end of file
diff --git a/Engine/Plugins/Animation/AnimationData/Source/AnimationData/Private/AnimSequencerDataModel.h b/Engine/Plugins/Animation/AnimationData/Source/AnimationData/Private/AnimSequencerDataModel.h
--- a/Engine/Plugins/Animation/AnimationData/Source/AnimationData/Private/AnimSequencerDataModel.h
+++ b/Engine/Plugins/Animation/AnimationData/Source/AnimationData/Private/AnimSequencerDataModel.h
@@ -82,6 +82,12 @@ public:
 	virtual const FTransformCurve* FindTransformCurve(const FAnimationCurveIdentifier& CurveIdentifier) const override;
 	virtual const FRichCurve* FindRichCurve(const FAnimationCurveIdentifier& CurveIdentifier) const override;
 	PRAGMA_ENABLE_DEPRECATION_WARNINGS
+	virtual bool IsValidBoneTrackName(const FName& TrackName) const override;
+	virtual FTransform GetBoneTrackTransform(FName TrackName, const FFrameNumber& FrameNumber) const override;
+	virtual void GetBoneTrackTransforms(FName TrackName, const TArray<FFrameNumber>& FrameNumbers, TArray<FTransform>& OutTransforms) const override;	
+	virtual void GetBoneTrackTransforms(FName TrackName, TArray<FTransform>& OutTransforms) const override;
+	virtual void GetBoneTracksTransform(const TArray<FName>& TrackNames, const FFrameNumber& FrameNumber, TArray<FTransform>& OutTransforms) const override;	
+	virtual FTransform EvaluateBoneTrackTransform(FName TrackName, const FFrameTime& FrameTime, const EAnimInterpolationType& Interpolation) const override;	
 	virtual const FAnimationCurveData& GetCurveData() const override;
 	virtual int32 GetNumberOfTransformCurves() const override;
 	virtual int32 GetNumberOfFloatCurves() const override;
@@ -105,19 +111,27 @@ public:
 	virtual FAnimDataModelModifiedDynamicEvent& GetModifiedDynamicEvent() override { return ModifiedEventDynamic; }
     virtual void Evaluate(FAnimationPoseData& InOutPoseData, const UE::Anim::DataModel::FEvaluationContext& EvaluationContext) const override;
 	virtual bool HasBeenPopulated() const override { return bPopulated; }
+	virtual void IterateBoneKeys(const FName& BoneName, TFunction<bool(const FVector3f& Pos, const FQuat4f&, const FVector3f, const FFrameNumber&)> IterationFunction) const override;
 protected:
 	virtual void OnNotify(const EAnimDataModelNotifyType& NotifyType, const FAnimDataModelNotifPayload& Payload) override final;
+
+	virtual void LockEvaluationAndModification() const override final
+	{
+		EvaluationLock.Lock();
+	}
+	virtual void UnlockEvaluationAndModification() const override final
+	{
+		EvaluationLock.Unlock();
+	}
 	/** End IAnimationDataModel overrides */
 	
 	/** Controller helper functionality */
 	FTransformCurve* FindMutableTransformCurveById(const FAnimationCurveIdentifier& CurveIdentifier);
 	FFloatCurve* FindMutableFloatCurveById(const FAnimationCurveIdentifier& CurveIdentifier);
 	FAnimCurveBase* FindMutableCurveById(const FAnimationCurveIdentifier& CurveIdentifier);
-	FBoneAnimationTrack* FindMutableBoneTrackByName(FName Name);	
 	FRichCurve* GetMutableRichCurve(const FAnimationCurveIdentifier& CurveIdentifier);
 
 	void GenerateLegacyCurveData();
-	void GenerateLegacyBoneData();
 
 	void ValidateData() const;
 	void ValidateSequencerData() const;
@@ -133,10 +147,11 @@ protected:
 	void InitializeFKControlRig(UFKControlRig* FKControlRig, USkeleton* Skeleton) const;
 	UControlRig* GetControlRig() const;
 	
-	void IterateTransformControlCurve(const FName& BoneName, TFunction<void(const FTransform&, const FFrameNumber&)> IterationFunction) const;
+	void IterateTransformControlCurve(const FName& BoneName, TFunction<void(const FTransform&, const FFrameNumber&)> IterationFunction, const TArray<FFrameNumber>* InFrameNumbers = nullptr) const;
 	void GenerateTransformKeysForControl(const FName& BoneName, TArray<FTransform>& InOutTransforms, TArray<FFrameNumber>& InOutFrameNumbers) const;
 	void RemoveOutOfDateControls() const;
-
+	
+	void GenerateTransformKeysForControl(const FName& BoneName, const TArray<FFrameNumber>& FrameNumbers, TArray<FTransform>& InOutTransforms) const;
 private:	
 	/** Dynamic delegate event allows scripting to register to any broadcast-ed notify. */
 	UPROPERTY(BlueprintAssignable, Transient, Category = AnimationDataModel, meta = (ScriptName = "ModifiedEvent", AllowPrivateAccess = "true"))
@@ -145,9 +160,6 @@ private:
 	TUniquePtr<IAnimationDataModel::FModelNotifier> Notifier;
 	UE::Anim::FAnimDataModelNotifyCollector Collector;
 
-	UPROPERTY(VisibleAnywhere, Category=AnimSequencer)
-	TArray<FBoneAnimationTrack> LegacyBoneAnimationTracks;
-
 	UPROPERTY(VisibleAnywhere, Category=AnimSequencer)
 	FAnimationCurveData LegacyCurveData;
 
@@ -174,3 +186,5 @@ private:
 
 	friend class UAnimSequencerController;
 };
+
+
diff --git a/Engine/Plugins/Animation/AnimationModifierLibrary/Source/AnimationModifierLibrary/Private/CopyBonesModifier.cpp b/Engine/Plugins/Animation/AnimationModifierLibrary/Source/AnimationModifierLibrary/Private/CopyBonesModifier.cpp
--- a/Engine/Plugins/Animation/AnimationModifierLibrary/Source/AnimationModifierLibrary/Private/CopyBonesModifier.cpp
+++ b/Engine/Plugins/Animation/AnimationModifierLibrary/Source/AnimationModifierLibrary/Private/CopyBonesModifier.cpp
@@ -39,18 +39,21 @@ void UCopyBonesModifier::OnApply_Implementation(UAnimSequence* Animation)
 			: SourceBoneName(InSourceBoneName), TargetBoneName(InTargetBoneName), SourceBoneIdx(InSourceBoneIdx), TargetBoneIdx(InTargetBoneIdx) {}
 	};
 
+	const USkeleton* Skeleton = Animation->GetSkeleton();
+	const FReferenceSkeleton& RefSkeleton = Skeleton->GetReferenceSkeleton();
+
 	// Validate input
 	TArray<FCopyBoneData> CopyBoneDataContainer;
 	CopyBoneDataContainer.Reserve(BonePairs.Num());
 	for (const FBoneReferencePair& Pair : BonePairs)
 	{
-		const int32 SourceBoneIdx = Model->GetBoneTrackIndexByName(Pair.SourceBone.BoneName);
+		const int32 SourceBoneIdx = RefSkeleton.FindBoneIndex(Pair.SourceBone.BoneName);
 		if (SourceBoneIdx == INDEX_NONE)
 		{
 			continue;
 		}
 
-		const int32 TargetBoneIdx = Model->GetBoneTrackIndexByName(Pair.TargetBone.BoneName);
+		const int32 TargetBoneIdx = RefSkeleton.FindBoneIndex(Pair.TargetBone.BoneName);
 		if (TargetBoneIdx == INDEX_NONE)
 		{
 			continue;
@@ -66,7 +69,7 @@ void UCopyBonesModifier::OnApply_Implementation(UAnimSequence* Animation)
 	TGuardValue<bool> ForceRootLockGuard(Animation->bForceRootLock, true);
 
 	// Start editing animation data
-	const bool bShouldTransact = false;
+	constexpr bool bShouldTransact = false;
 	Controller.OpenBracket(LOCTEXT("CopyBonesModifier_Bracket", "Updating bones"), bShouldTransact);
 
 	// Get the transform of all the source bones in the desired space
diff --git a/Engine/Plugins/Animation/AnimationModifierLibrary/Source/AnimationModifierLibrary/Private/ReOrientRootBoneModifier.cpp b/Engine/Plugins/Animation/AnimationModifierLibrary/Source/AnimationModifierLibrary/Private/ReOrientRootBoneModifier.cpp
--- a/Engine/Plugins/Animation/AnimationModifierLibrary/Source/AnimationModifierLibrary/Private/ReOrientRootBoneModifier.cpp
+++ b/Engine/Plugins/Animation/AnimationModifierLibrary/Source/AnimationModifierLibrary/Private/ReOrientRootBoneModifier.cpp
@@ -54,24 +54,25 @@ void UReOrientRootBoneModifier::ReOrientRootBone_Internal(UAnimSequence* Animati
 	}
 
 	// Collect tracks for the bones we are going to modify (root and direct children of the root)
-	TArray<const FBoneAnimationTrack*> BonesTracks;
-	BonesTracks.Add(&Model->GetBoneTrackByIndex(0));
-	for (int32 BoneIdx = 1; BoneIdx < RefSkeleton.GetNum(); BoneIdx++)
+	const FName RootBoneName = RefSkeleton.GetBoneName(0);
+	TArray<FName> RootBoneChildBones;
+
+	TArray<int32> ChildBoneIndices;
+	RefSkeleton.GetDirectChildBones(0, ChildBoneIndices);
+
+	for (const int32& ChildBoneIndex : ChildBoneIndices)
 	{
-		if (RefSkeleton.GetParentIndex(BoneIdx) == 0)
+		const FName ChildBoneName = RefSkeleton.GetBoneName(ChildBoneIndex);
+		if (Model->IsValidBoneTrackName(ChildBoneName))
 		{
-			const int32 BoneTrackIdx = Model->GetBoneTrackIndexByName(RefSkeleton.GetBoneName(BoneIdx));
-			if(BoneTrackIdx != INDEX_NONE)
-			{
-				BonesTracks.Add(&Model->GetBoneTrackByIndex(BoneTrackIdx));
-			}
+			RootBoneChildBones.Add(ChildBoneName);
 		}
 	}
 	
 	// Start editing animation data
 	const bool bShouldTransact = false;
 	Controller.OpenBracket(LOCTEXT("ReorientingRootBone_Bracket", "Reorienting root bone"), bShouldTransact);
-	
+
 	// For each key in the animation
 	const int32 Num = Model->GetNumberOfKeys();
 	for (int32 AnimKey = 0; AnimKey < Num; AnimKey++)
@@ -79,22 +80,20 @@ void UReOrientRootBoneModifier::ReOrientRootBone_Internal(UAnimSequence* Animati
 		const FInt32Range KeyRangeToSet(AnimKey, AnimKey + 1);
 
 		// Reorient the root bone
-		FTransform RootTransformOriginal;
-		FAnimationUtils::ExtractTransformForFrameFromTrack(BonesTracks[0]->InternalTrackData, AnimKey, RootTransformOriginal);
+		const FTransform RootTransformOriginal = Model->EvaluateBoneTrackTransform(RootBoneName, AnimKey, EAnimInterpolationType::Step);
 
 		FTransform RootTransformNew = RootTransformOriginal;
 		RootTransformNew.SetRotation(RootTransformOriginal.GetRotation() * Quat);
-		Controller.UpdateBoneTrackKeys(BonesTracks[0]->Name, KeyRangeToSet, { RootTransformNew.GetLocation() }, { RootTransformNew.GetRotation() }, { RootTransformNew.GetScale3D() });
+		Controller.UpdateBoneTrackKeys(RootBoneName, KeyRangeToSet, { RootTransformNew.GetLocation() }, { RootTransformNew.GetRotation() }, { RootTransformNew.GetScale3D() });
 
 		// Now the mesh is facing the wrong axis. Update direct children of the root with the local space transform that puts them back to where they were originally
-		for (int32 Idx = 1; Idx < BonesTracks.Num(); Idx++)
+
+		for (const FName& ChildBoneName : RootBoneChildBones)
 		{
-			FTransform OtherBoneTransform;
-			FAnimationUtils::ExtractTransformForFrameFromTrack(BonesTracks[Idx]->InternalTrackData, AnimKey, OtherBoneTransform);
-			
+			FTransform OtherBoneTransform = Model->EvaluateBoneTrackTransform(ChildBoneName, AnimKey, EAnimInterpolationType::Step);
 			OtherBoneTransform = OtherBoneTransform * RootTransformOriginal;
 			OtherBoneTransform = OtherBoneTransform.GetRelativeTransform(RootTransformNew);
-			Controller.UpdateBoneTrackKeys(BonesTracks[Idx]->Name, KeyRangeToSet, { OtherBoneTransform.GetLocation() }, { OtherBoneTransform.GetRotation() }, { OtherBoneTransform.GetScale3D() });
+			Controller.UpdateBoneTrackKeys(ChildBoneName, KeyRangeToSet, { OtherBoneTransform.GetLocation() }, { OtherBoneTransform.GetRotation() }, { OtherBoneTransform.GetScale3D() });
 		}
 	}
 
diff --git a/Engine/Plugins/Animation/IKRig/Source/IKRigEditor/Private/RetargetEditor/IKRetargetBatchOperation.cpp b/Engine/Plugins/Animation/IKRig/Source/IKRigEditor/Private/RetargetEditor/IKRetargetBatchOperation.cpp
--- a/Engine/Plugins/Animation/IKRig/Source/IKRigEditor/Private/RetargetEditor/IKRetargetBatchOperation.cpp
+++ b/Engine/Plugins/Animation/IKRig/Source/IKRigEditor/Private/RetargetEditor/IKRetargetBatchOperation.cpp
@@ -395,7 +395,7 @@ void UIKRetargetBatchOperation::ConvertAnimation(
 			const FName& TargetBoneName = TargetBoneNames[TargetBoneIndex];
 
 			const FRawAnimSequenceTrack& RawTrack = BoneTracks[TargetBoneIndex];
-			TargetSeqController.AddBoneTrack(TargetBoneName, bShouldTransact);
+			TargetSeqController.AddBoneCurve(TargetBoneName, bShouldTransact);
 			TargetSeqController.SetBoneTrackKeys(TargetBoneName, RawTrack.PosKeys, RawTrack.RotKeys, RawTrack.ScaleKeys);
 		}
 
diff --git a/Engine/Plugins/Animation/MLDeformer/MLDeformerFramework/Source/MLDeformerFrameworkEditor/Private/MLDeformerEditorModel.cpp b/Engine/Plugins/Animation/MLDeformer/MLDeformerFramework/Source/MLDeformerFrameworkEditor/Private/MLDeformerEditorModel.cpp
--- a/Engine/Plugins/Animation/MLDeformer/MLDeformerFramework/Source/MLDeformerFrameworkEditor/Private/MLDeformerEditorModel.cpp
+++ b/Engine/Plugins/Animation/MLDeformer/MLDeformerFramework/Source/MLDeformerFrameworkEditor/Private/MLDeformerEditorModel.cpp
@@ -1098,25 +1098,28 @@ namespace UE::MLDeformer
 		TArray<FName> AnimatedBoneList;
 		const FReferenceSkeleton& RefSkeleton = Skeleton->GetReferenceSkeleton();
 		const int32 NumBones = RefSkeleton.GetNum();
+
+		TArray<FName> TrackNames;
+		DataModel->GetBoneTrackNames(TrackNames);
+		TArray<FTransform> BoneTransforms;		
 		for (int32 Index = 0; Index < NumBones; ++Index)
 		{
 			const FName BoneName = RefSkeleton.GetBoneName(Index);
-			const int32 BoneTrackIndex = DataModel->GetBoneTrackIndexByName(BoneName);
-			if (BoneTrackIndex == INDEX_NONE)
+			if (!TrackNames.Contains(BoneName))
 			{
 				continue;
 			}
 
 			// Check if there is actually animation data.
-			const FBoneAnimationTrack& BoneAnimTrack = DataModel->GetBoneTrackByIndex(BoneTrackIndex);
-			const TArray<FQuat4f>& Rotations = BoneAnimTrack.InternalTrackData.RotKeys;
+			BoneTransforms.Reset();
+			DataModel->GetBoneTrackTransforms(BoneName, BoneTransforms);
 			bool bIsAnimated = false;
-			if (!Rotations.IsEmpty())
+			if (!BoneTransforms.IsEmpty())
 			{
-				const FQuat4f FirstQuat = Rotations[0];
-				for (const FQuat4f KeyValue : Rotations)
+				const FQuat FirstQuat = BoneTransforms[0].GetRotation();
+				for (const FTransform& KeyValue : BoneTransforms)
 				{
-					if (!KeyValue.Equals(FirstQuat))
+					if (!KeyValue.GetRotation().Equals(FirstQuat))
 					{
 						bIsAnimated = true;
 						break;
diff --git a/Engine/Plugins/Animation/MLDeformer/NearestNeighborModel/Source/NearestNeighborModelEditor/Private/NearestNeighborEditorModel.cpp b/Engine/Plugins/Animation/MLDeformer/NearestNeighborModel/Source/NearestNeighborModelEditor/Private/NearestNeighborEditorModel.cpp
--- a/Engine/Plugins/Animation/MLDeformer/NearestNeighborModel/Source/NearestNeighborModelEditor/Private/NearestNeighborEditorModel.cpp
+++ b/Engine/Plugins/Animation/MLDeformer/NearestNeighborModel/Source/NearestNeighborModelEditor/Private/NearestNeighborEditorModel.cpp
@@ -561,11 +561,11 @@ namespace UE::NearestNeighborModel
 		Controller.SetNumberOfFrames(NumKeys - 1);
 		Controller.SetFrameRate(FFrameRate(30, 1));
 
-		const TArray<FBoneAnimationTrack>& DefaultTracks = DefaultAnim->GetDataModel()->GetBoneAnimationTracks();
-		const int32 NumTracks = DefaultTracks.Num();
-		for (int32 TrackIndex = 0; TrackIndex < NumTracks; TrackIndex++)
+		TArray<FName> TrackNames;		
+		DefaultAnim->GetDataModel()->GetBoneTrackNames(TrackNames);
+		const int32 NumTracks = TrackNames.Num();
+		for (const FName& TrackName : TrackNames)
 		{
-			const FName TrackName = DefaultTracks[TrackIndex].Name;
 			TArray<FVector3f> PosKeys;
 			TArray<FQuat4f> RotKeys;
 			TArray<FVector3f> ScaleKeys;
@@ -588,18 +588,13 @@ namespace UE::NearestNeighborModel
 					UE_LOG(LogNearestNeighborModel, Error, TEXT("CreateAnimOfClusterCenters: PickedAnim %d is null."), PickedAnimId);
 					return TTuple<UAnimSequence*, uint8>(nullptr, EUpdateResult::ERROR);
 				}
-				const FBoneAnimationTrack& Track = PickedAnim->GetDataModel()->GetBoneTrackByName(TrackName);
-				const FRawAnimSequenceTrack& RawTrack = Track.InternalTrackData;
-				if (PickedFrame < 0 || PickedFrame >= RawTrack.PosKeys.Num())
-				{
-					UE_LOG(LogNearestNeighborModel, Error, TEXT("CreateAnimOfClusterCenters: PickedFrame %d is out of range[0, %d)."), PickedFrame, RawTrack.PosKeys.Num());
-					return TTuple<UAnimSequence*, uint8>(nullptr, EUpdateResult::ERROR);
-				}
-				PosKeys[KeyIndex] = RawTrack.PosKeys[PickedFrame];
-				RotKeys[KeyIndex] = RawTrack.RotKeys[PickedFrame];
-				ScaleKeys[KeyIndex] = RawTrack.ScaleKeys[PickedFrame];
+
+				const FTransform BoneTransform = PickedAnim->GetDataModel()->GetBoneTrackTransform(TrackName, PickedFrame);				
+				PosKeys[KeyIndex] = FVector3f(BoneTransform.GetLocation());
+				RotKeys[KeyIndex] = FQuat4f(BoneTransform.GetRotation());
+				ScaleKeys[KeyIndex] = FVector3f(BoneTransform.GetScale3D());
 			}
-			Controller.AddBoneTrack(TrackName);
+			Controller.AddBoneCurve(TrackName);
 			Controller.SetBoneTrackKeys(TrackName, PosKeys, RotKeys, ScaleKeys);
 		}
 		Controller.NotifyPopulated();
