Temporarily disable raw data GUID caching, as compressed anim data results from DDC for the generated keys are invalid

diff --git a/Engine/Plugins/Animation/AnimationData/Source/AnimationData/Private/AnimSequencerController.cpp b/Engine/Plugins/Animation/AnimationData/Source/AnimationData/Private/AnimSequencerController.cpp
--- a/Engine/Plugins/Animation/AnimationData/Source/AnimationData/Private/AnimSequencerController.cpp
+++ b/Engine/Plugins/Animation/AnimationData/Source/AnimationData/Private/AnimSequencerController.cpp
@@ -2496,7 +2496,8 @@ void UAnimSequencerController::InitializeModel()
 				if (AnimSequence->GetLinkerCustomVersion(FUE5MainStreamObjectVersion::GUID) < FUE5MainStreamObjectVersion::IntroducingAnimationDataModel)
 				{
 					PRAGMA_DISABLE_DEPRECATION_WARNINGS
-					Model->CachedRawDataGUID = AnimSequence->GenerateGuidFromRawData();
+					// Disabled for now as DDC had invalid data cached causing collapsed animations
+					//Model->CachedRawDataGUID = AnimSequence->GenerateGuidFromRawData();
 					PRAGMA_ENABLE_DEPRECATION_WARNINGS
 				}
 				
