AnimationData: Sanitize element names for Control Rig

diff --git a/Engine/Plugins/Animation/AnimationData/Source/AnimationData/Private/AnimSequencerController.cpp b/Engine/Plugins/Animation/AnimationData/Source/AnimationData/Private/AnimSequencerController.cpp
--- a/Engine/Plugins/Animation/AnimationData/Source/AnimationData/Private/AnimSequencerController.cpp
+++ b/Engine/Plugins/Animation/AnimationData/Source/AnimationData/Private/AnimSequencerController.cpp
@@ -1432,8 +1432,8 @@ bool UAnimSequencerController::SetCurveAttributes(const FAnimationCurveIdentifie
 				const URigHierarchy* Hierarchy = FKRig->GetHierarchy();
 				if (Hierarchy ||  IgnoreSkeletonValidation())
 				{
-					const FRigElementKey CurveKey(CurveId.InternalName.DisplayName, ERigElementType::Curve);
-					const FRigElementKey CurveControlKey(UFKControlRig::GetControlName(CurveId.InternalName.DisplayName, ERigElementType::Curve), ERigElementType::Control);
+					const FRigElementKey CurveKey(URigHierarchy::GetSanitizedName(CurveId.InternalName.DisplayName.ToString()), ERigElementType::Curve);
+					const FRigElementKey CurveControlKey(UFKControlRig::GetControlName(CurveKey.Name, ERigElementType::Curve), ERigElementType::Control);
 
 					const bool bContainsCurve = Hierarchy ? Hierarchy->Contains(CurveKey) : IgnoreSkeletonValidation();
 					const bool bContainsCurveControl = Hierarchy ? Hierarchy->Contains(CurveControlKey) : IgnoreSkeletonValidation();
@@ -1667,8 +1667,8 @@ bool UAnimSequencerController::SetBoneTrackKeys(FName BoneName, const TArray<FVe
 				{
 					if (URigHierarchy* Hierarchy = ControlRig->GetHierarchy())
 					{
-						const FRigElementKey BoneKey(BoneName, ERigElementType::Bone);
-						const FRigElementKey BoneNameControlKey(UFKControlRig::GetControlName(BoneName, ERigElementType::Bone), ERigElementType::Control);
+						const FRigElementKey BoneKey(URigHierarchy::GetSanitizedName(BoneName.ToString()), ERigElementType::Bone);
+						const FRigElementKey BoneNameControlKey(UFKControlRig::GetControlName(BoneKey.Name, ERigElementType::Bone), ERigElementType::Control);
 						
 						const bool bContainsBone = Hierarchy->Contains(BoneKey);
 						const bool bContainsBoneControl = Hierarchy->Contains(BoneNameControlKey);
@@ -2576,8 +2576,8 @@ bool UAnimSequencerController::AddBoneControl(const FName& BoneName) const
 		{
 			if(URigHierarchy* Hierarchy = FKRig->GetHierarchy())
 			{
-				const FRigElementKey BoneKey(BoneName, ERigElementType::Bone);
-				const FRigElementKey BoneNameControlKey(UFKControlRig::GetControlName(BoneName, ERigElementType::Bone), ERigElementType::Control);
+				const FRigElementKey BoneKey(URigHierarchy::GetSanitizedName(BoneName.ToString()), ERigElementType::Bone);
+				const FRigElementKey BoneNameControlKey(UFKControlRig::GetControlName(BoneKey.Name, ERigElementType::Bone), ERigElementType::Control);
 
 				const bool bContainsBone = Hierarchy->Contains(BoneKey);
 				const bool bContainsBoneControl = Hierarchy->Contains(BoneNameControlKey);
@@ -2681,7 +2681,7 @@ bool UAnimSequencerController::RemoveBoneControl(const FName& BoneName) const
 		{
 			if(URigHierarchy* Hierarchy = FKRig->GetHierarchy())
 			{
-				const FRigElementKey BoneNameControlKey(UFKControlRig::GetControlName(BoneName, ERigElementType::Bone), ERigElementType::Control);
+				const FRigElementKey BoneNameControlKey(UFKControlRig::GetControlName(URigHierarchy::GetSanitizedName(BoneName.ToString()), ERigElementType::Bone), ERigElementType::Control);
 				if(Hierarchy->Contains(BoneNameControlKey))
 				{
 					if(URigHierarchyController* HierarchyController = Hierarchy->GetController())
@@ -2732,8 +2732,8 @@ bool UAnimSequencerController::SetBoneCurveKeys(const FName& BoneName, const TAr
 			const URigHierarchy* Hierarchy = FKRig->GetHierarchy();
 			if(Hierarchy || IgnoreSkeletonValidation())
 			{
-				const FRigElementKey BoneKey(BoneName, ERigElementType::Bone);
-				const FRigElementKey BoneNameControlKey(UFKControlRig::GetControlName(BoneName, ERigElementType::Bone), ERigElementType::Control);
+				const FRigElementKey BoneKey(URigHierarchy::GetSanitizedName(BoneName.ToString()), ERigElementType::Bone);
+				const FRigElementKey BoneNameControlKey(UFKControlRig::GetControlName(BoneKey.Name, ERigElementType::Bone), ERigElementType::Control);
 
 				const bool bContainsBone = Hierarchy ? Hierarchy->Contains(BoneKey) : IgnoreSkeletonValidation();
 				const bool bContainsBoneControl = Hierarchy ? Hierarchy->Contains(BoneNameControlKey) : IgnoreSkeletonValidation();
@@ -2923,8 +2923,8 @@ bool UAnimSequencerController::UpdateBoneCurveKeys(const FName& BoneName, const
 		{
 			if(URigHierarchy* Hierarchy = FKRig->GetHierarchy())
 			{
-				const FRigElementKey BoneKey(BoneName, ERigElementType::Bone);
-				const FRigElementKey BoneNameControlKey(UFKControlRig::GetControlName(BoneName, ERigElementType::Bone), ERigElementType::Control);
+				const FRigElementKey BoneKey(URigHierarchy::GetSanitizedName(BoneName.ToString()), ERigElementType::Bone);
+				const FRigElementKey BoneNameControlKey(UFKControlRig::GetControlName(BoneKey.Name, ERigElementType::Bone), ERigElementType::Control);
 
 				const bool bContainsBone = Hierarchy->Contains(BoneKey);
 				const bool bContainsBoneControl = Hierarchy->Contains(BoneNameControlKey);
@@ -3019,8 +3019,8 @@ bool UAnimSequencerController::RemoveBoneCurveKey(const FName& BoneName, float T
 			{
 				if(const URigHierarchy* Hierarchy = FKRig->GetHierarchy())
 				{
-					const FRigElementKey BoneKey(BoneName, ERigElementType::Bone);
-					const FRigElementKey BoneNameControlKey(UFKControlRig::GetControlName(BoneName, ERigElementType::Bone), ERigElementType::Control);
+					const FRigElementKey BoneKey(URigHierarchy::GetSanitizedName(BoneName.ToString()), ERigElementType::Bone);
+					const FRigElementKey BoneNameControlKey(UFKControlRig::GetControlName(BoneKey.Name, ERigElementType::Bone), ERigElementType::Control);
 	
 					const bool bContainsBone = Hierarchy->Contains(BoneKey);
 					const bool bContainsBoneControl = Hierarchy->Contains(BoneNameControlKey);
@@ -3092,8 +3092,8 @@ bool UAnimSequencerController::AddCurveControl(const FName& CurveName) const
 		{
 			if(URigHierarchy* Hierarchy = FKRig->GetHierarchy())
 			{				
-				const FRigElementKey CurveKey(CurveName, ERigElementType::Curve);
-				const FRigElementKey CurveControlKey(UFKControlRig::GetControlName(CurveName, ERigElementType::Curve), ERigElementType::Control);
+				const FRigElementKey CurveKey(URigHierarchy::GetSanitizedName(CurveName.ToString()), ERigElementType::Curve);
+				const FRigElementKey CurveControlKey(UFKControlRig::GetControlName(CurveKey.Name, ERigElementType::Curve), ERigElementType::Control);
 
 				const bool bContainsCurve = Hierarchy->Contains(CurveKey);
 				const bool bContainsCurveControl = Hierarchy->Contains(CurveControlKey);
@@ -3159,7 +3159,7 @@ bool UAnimSequencerController::RenameCurveControl(const FName& CurveName, const
 			{
 				if (URigHierarchyController* Controller = RigHierarchy->GetController())
 				{
-					const FRigElementKey CurveKey(CurveName, ERigElementType::Curve);
+					const FRigElementKey CurveKey(URigHierarchy::GetSanitizedName(CurveName.ToString()), ERigElementType::Curve);
 					const FRigElementKey CurveControlKey(UFKControlRig::GetControlName(CurveKey.Name, ERigElementType::Curve), ERigElementType::Control);
 					
 					// Rename the curve element itself
@@ -3171,11 +3171,11 @@ bool UAnimSequencerController::RenameCurveControl(const FName& CurveName, const
 							// Rename the control for the curve value
 							if (RigHierarchy->Contains(CurveControlKey))
 							{											
-								const FRigElementKey NewCurveControlKey = Controller->RenameElement(CurveControlKey, UFKControlRig::GetControlName(NewCurveName, ERigElementType::Curve), false);
+								const FRigElementKey NewCurveControlKey = Controller->RenameElement(CurveControlKey, UFKControlRig::GetControlName(NewCurveKey.Name, ERigElementType::Curve), false);
 								if (NewCurveControlKey.IsValid())
 								{
 									FRigControlElement* ControlElement = RigHierarchy->Find<FRigControlElement>(NewCurveControlKey);
-									ControlElement->Settings.DisplayName = FName(*(NewCurveName.ToString() + TEXT(" Curve")));
+									ControlElement->Settings.DisplayName = FName(*(NewCurveControlKey.Name.ToString() + TEXT(" Curve")));
 
 									// Rename the curve driving the control value
 									FScalarParameterNameAndCurve* ParameterCurvePair = Section->GetScalarParameterNamesAndCurves().FindByPredicate([CurveControlKey](const FScalarParameterNameAndCurve& Parameter)
@@ -3185,7 +3185,7 @@ bool UAnimSequencerController::RenameCurveControl(const FName& CurveName, const
 					
 									if(ParameterCurvePair)
 									{
-										ParameterCurvePair->ParameterName = UFKControlRig::GetControlName(NewCurveName, ERigElementType::Curve);
+										ParameterCurvePair->ParameterName = UFKControlRig::GetControlName(NewCurveControlKey.Name, ERigElementType::Curve);
 									}
 
 									return true;
@@ -3245,8 +3245,8 @@ bool UAnimSequencerController::RemoveCurveControl(const FName& CurveName) const
 		{
 			if(URigHierarchy* Hierarchy = FKRig->GetHierarchy())
 			{									
-				const FRigElementKey CurveKey(CurveName, ERigElementType::Curve);
-				const FRigElementKey CurveControlKey(UFKControlRig::GetControlName(CurveName, ERigElementType::Curve), ERigElementType::Control);
+				const FRigElementKey CurveKey(URigHierarchy::GetSanitizedName(CurveName.ToString()), ERigElementType::Curve);
+				const FRigElementKey CurveControlKey(UFKControlRig::GetControlName(CurveKey.Name, ERigElementType::Curve), ERigElementType::Control);
 
 				const bool bContainsCurve = Hierarchy->Contains(CurveKey);
 				const bool bContainsCurveControl = Hierarchy->Contains(CurveControlKey);
@@ -3311,8 +3311,8 @@ bool UAnimSequencerController::SetCurveControlKeys(const FName& CurveName, const
 
 			if (Hierarchy ||  IgnoreSkeletonValidation())
 			{
-				const FRigElementKey CurveKey(CurveName, ERigElementType::Curve);
-				const FRigElementKey CurveControlKey(UFKControlRig::GetControlName(CurveName, ERigElementType::Curve), ERigElementType::Control);
+				const FRigElementKey CurveKey(URigHierarchy::GetSanitizedName(CurveName.ToString()), ERigElementType::Curve);
+				const FRigElementKey CurveControlKey(UFKControlRig::GetControlName(CurveKey.Name, ERigElementType::Curve), ERigElementType::Control);
 
 				const bool bContainsCurve = Hierarchy ? Hierarchy->Contains(CurveKey) : IgnoreSkeletonValidation();
 				const bool bContainsCurveControl = Hierarchy ? Hierarchy->Contains(CurveControlKey) : IgnoreSkeletonValidation();
@@ -3397,8 +3397,8 @@ bool UAnimSequencerController::SetCurveControlKey(const FName& CurveName, const
 		{
 			if(const URigHierarchy* Hierarchy = FKRig->GetHierarchy())
 			{
-				const FRigElementKey CurveKey(CurveName, ERigElementType::Curve);
-				const FRigElementKey CurveControlKey(UFKControlRig::GetControlName(CurveName, ERigElementType::Curve), ERigElementType::Control);
+				const FRigElementKey CurveKey(URigHierarchy::GetSanitizedName(CurveName.ToString()), ERigElementType::Curve);
+				const FRigElementKey CurveControlKey(UFKControlRig::GetControlName(CurveKey.Name, ERigElementType::Curve), ERigElementType::Control);
 				const bool bContainsCurve = Hierarchy->Contains(CurveKey);
 				const bool bContainsCurveControl = Hierarchy->Contains(CurveControlKey);
 				
@@ -3472,8 +3472,8 @@ bool UAnimSequencerController::RemoveCurveControlKey(const FName& CurveName, flo
 			if(const URigHierarchy* Hierarchy = FKRig->GetHierarchy())
 			{
 				const FFrameNumber FrameNumber = Model->MovieScene->GetTickResolution().AsFrameNumber(Time);
-				const FRigElementKey CurveKey(CurveName, ERigElementType::Curve);
-				const FRigElementKey CurveControlKey(UFKControlRig::GetControlName(CurveName, ERigElementType::Curve), ERigElementType::Control);
+				const FRigElementKey CurveKey(URigHierarchy::GetSanitizedName(CurveName.ToString()), ERigElementType::Curve);
+				const FRigElementKey CurveControlKey(UFKControlRig::GetControlName(CurveKey.Name, ERigElementType::Curve), ERigElementType::Control);
 				const bool bContainsCurve = Hierarchy->Contains(CurveKey);
 				const bool bContainsCurveControl = Hierarchy->Contains(CurveControlKey);
 
@@ -3537,8 +3537,8 @@ bool UAnimSequencerController::DuplicateCurveControl(const FName& CurveName, con
 			{
 				if (URigHierarchyController* HierarchyController = Hierarchy->GetController())
 				{
-					const FRigElementKey NewCurveKey(DuplicateCurveName, ERigElementType::Curve);
-					const FRigElementKey NewCurveControlKey(UFKControlRig::GetControlName(DuplicateCurveName, ERigElementType::Curve), ERigElementType::Control);
+					const FRigElementKey NewCurveKey(URigHierarchy::GetSanitizedName(DuplicateCurveName.ToString()), ERigElementType::Curve);
+					const FRigElementKey NewCurveControlKey(UFKControlRig::GetControlName(NewCurveKey.Name, ERigElementType::Curve), ERigElementType::Control);
 
 					const bool bAlreadyContainsCurve = Hierarchy->Contains(NewCurveKey);
 					const bool bAlreadyContainsCurveControl = Hierarchy->Contains(NewCurveControlKey);
@@ -3557,7 +3557,7 @@ bool UAnimSequencerController::DuplicateCurveControl(const FName& CurveName, con
 						// Rename the curve driving the control value
 						const FScalarParameterNameAndCurve* ParameterCurvePair = Section->GetScalarParameterNamesAndCurves().FindByPredicate([CurveName](const FScalarParameterNameAndCurve& Parameter)
 						{
-							return Parameter.ParameterName == UFKControlRig::GetControlName(CurveName, ERigElementType::Curve);
+							return Parameter.ParameterName == UFKControlRig::GetControlName(URigHierarchy::GetSanitizedName(CurveName.ToString()), ERigElementType::Curve);
 						});
 				
 						if(ParameterCurvePair)
