//UE5/Main - UE.EditorAutomation(RunTest=Editor) Linux vulkan - Curves - invalid number of keys
#fix ensure we use the correct sanitized FName during curve (control) renaming

diff --git a/Engine/Plugins/Animation/AnimationData/Source/AnimationData/Private/AnimSequencerController.cpp b/Engine/Plugins/Animation/AnimationData/Source/AnimationData/Private/AnimSequencerController.cpp
--- a/Engine/Plugins/Animation/AnimationData/Source/AnimationData/Private/AnimSequencerController.cpp
+++ b/Engine/Plugins/Animation/AnimationData/Source/AnimationData/Private/AnimSequencerController.cpp
@@ -3161,21 +3161,22 @@ bool UAnimSequencerController::RenameCurveControl(const FName& CurveName, const
 				{
 					const FRigElementKey CurveKey(URigHierarchy::GetSanitizedName(CurveName.ToString()), ERigElementType::Curve);
 					const FRigElementKey CurveControlKey(UFKControlRig::GetControlName(CurveKey.Name, ERigElementType::Curve), ERigElementType::Control);
-					
+
+					const FName SanitizedNewCurveName = URigHierarchy::GetSanitizedName(NewCurveName.ToString());
 					// Rename the curve element itself
 					if (RigHierarchy->Contains(CurveKey))
 					{											
-						const FRigElementKey NewCurveKey = Controller->RenameElement(CurveKey, NewCurveName, false);
+						const FRigElementKey NewCurveKey = Controller->RenameElement(CurveKey, SanitizedNewCurveName, false);
 						if (NewCurveKey.IsValid())
 						{
 							// Rename the control for the curve value
 							if (RigHierarchy->Contains(CurveControlKey))
-							{											
+							{
 								const FRigElementKey NewCurveControlKey = Controller->RenameElement(CurveControlKey, UFKControlRig::GetControlName(NewCurveKey.Name, ERigElementType::Curve), false);
 								if (NewCurveControlKey.IsValid())
 								{
 									FRigControlElement* ControlElement = RigHierarchy->Find<FRigControlElement>(NewCurveControlKey);
-									ControlElement->Settings.DisplayName = FName(*(NewCurveControlKey.Name.ToString() + TEXT(" Curve")));
+									ControlElement->Settings.DisplayName = FName(*(NewCurveKey.Name.ToString() + TEXT(" Curve")));
 
 									// Rename the curve driving the control value
 									FScalarParameterNameAndCurve* ParameterCurvePair = Section->GetScalarParameterNamesAndCurves().FindByPredicate([CurveControlKey](const FScalarParameterNameAndCurve& Parameter)
@@ -3185,7 +3186,7 @@ bool UAnimSequencerController::RenameCurveControl(const FName& CurveName, const
 					
 									if(ParameterCurvePair)
 									{
-										ParameterCurvePair->ParameterName = UFKControlRig::GetControlName(NewCurveControlKey.Name, ERigElementType::Curve);
+										ParameterCurvePair->ParameterName = NewCurveControlKey.Name;
 									}
 
 									return true;
