Revert behaviour, ensuring that correct ReferenceSkeleton is used to retrieve Skeleton to Compact pose mapping

diff --git a/Engine/Plugins/Animation/AnimationData/Source/AnimationData/Private/AnimSequencerDataModel.cpp b/Engine/Plugins/Animation/AnimationData/Source/AnimationData/Private/AnimSequencerDataModel.cpp
--- a/Engine/Plugins/Animation/AnimationData/Source/AnimationData/Private/AnimSequencerDataModel.cpp
+++ b/Engine/Plugins/Animation/AnimationData/Source/AnimationData/Private/AnimSequencerDataModel.cpp
@@ -1110,15 +1110,17 @@ void UAnimationSequencerDataModel::GeneratePoseData(UControlRig* ControlRig, FAn
 			// Populate bone/curve elements to Pose/Curve indices
 			{
 				QUICK_SCOPE_CYCLE_COUNTER(STAT_GetMappings);
-				const FReferenceSkeleton& RefSkeleton = RequiredBones.GetReferenceSkeleton();
+				const FReferenceSkeleton& MeshRefSkeleton = RequiredBones.GetReferenceSkeleton();
+				// Called during compression that can occur while GC is in progress, marking weakptrs as unreachable temporarily
+				const FReferenceSkeleton& SkeletonRefSkeleton = RequiredBones.GetSkeletonAsset(true)->GetReferenceSkeleton();
 
-				RigHierarchy->ForEach<FRigBoneElement>([&RefSkeleton, &RequiredBones, &RigHierarchy, &RetargetingScope, &RigPose](const FRigBoneElement* BoneElement) -> bool
+				RigHierarchy->ForEach<FRigBoneElement>([&MeshRefSkeleton, &RequiredBones, &SkeletonRefSkeleton, &RigHierarchy, &RetargetingScope, &RigPose](const FRigBoneElement* BoneElement) -> bool
 				{
 					const FName& BoneName = BoneElement->GetName();
-					const int32 BoneIndex = RefSkeleton.FindBoneIndex(BoneName);
+					const int32 BoneIndex = MeshRefSkeleton.FindBoneIndex(BoneName);
 					if (BoneIndex != INDEX_NONE)
 					{
-						const int32 SkeletonBoneIndex = RefSkeleton.FindBoneIndex(BoneName);
+						const int32 SkeletonBoneIndex = SkeletonRefSkeleton.FindBoneIndex(BoneName);
 						const FCompactPoseBoneIndex CompactPoseBoneIndex = RequiredBones.GetCompactPoseIndexFromSkeletonIndex(SkeletonBoneIndex);
 						if (CompactPoseBoneIndex != INDEX_NONE)
 						{
