Ensure additive Transform Curves are still incorporated into animation data GUID, until they are incorporated into Sequencer Data itself

diff --git a/Engine/Plugins/Animation/AnimationData/Source/AnimationData/Private/AnimSequencerDataModel.cpp b/Engine/Plugins/Animation/AnimationData/Source/AnimationData/Private/AnimSequencerDataModel.cpp
--- a/Engine/Plugins/Animation/AnimationData/Source/AnimationData/Private/AnimSequencerDataModel.cpp
+++ b/Engine/Plugins/Animation/AnimationData/Source/AnimationData/Private/AnimSequencerDataModel.cpp
@@ -632,6 +632,31 @@ FGuid UAnimationSequencerDataModel::GenerateGuid() const
 			UpdateSHAWithArray(Attribute.Curve.GetConstRefOfKeys());
 		}
 
+		auto UpdateWithFloatCurve = [&UpdateWithData, &UpdateSHAWithArray](const FRichCurve& Curve)
+		{
+			UpdateWithData(Curve.DefaultValue);
+			UpdateSHAWithArray(Curve.GetConstRefOfKeys());
+			UpdateWithData(Curve.PreInfinityExtrap);
+			UpdateWithData(Curve.PostInfinityExtrap);
+		};
+
+		for (const FTransformCurve& Curve : GetTransformCurves())
+		{
+			UpdateWithData(Curve.Name.UID);
+
+			auto UpdateWithComponent = [&UpdateWithFloatCurve](const FVectorCurve& VectorCurve)
+			{
+				for (int32 ChannelIndex = 0; ChannelIndex < 3; ++ChannelIndex)
+				{
+					UpdateWithFloatCurve(VectorCurve.FloatCurves[ChannelIndex]);
+				}
+			};
+	
+			UpdateWithComponent(Curve.TranslationCurve);
+			UpdateWithComponent(Curve.RotationCurve);
+			UpdateWithComponent(Curve.ScaleCurve);
+		}
+
 		Sha.Final();
 
 		uint32 Hash[5];
