Static analysis fix - nullptr check

diff --git a/Engine/Plugins/Animation/AnimationData/Source/AnimationData/Private/AnimSequencerController.cpp b/Engine/Plugins/Animation/AnimationData/Source/AnimationData/Private/AnimSequencerController.cpp
--- a/Engine/Plugins/Animation/AnimationData/Source/AnimationData/Private/AnimSequencerController.cpp
+++ b/Engine/Plugins/Animation/AnimationData/Source/AnimationData/Private/AnimSequencerController.cpp
@@ -2398,7 +2398,7 @@ void UAnimSequencerController::PopulateWithExistingModel(TScriptInterface<IAnima
 		if (const UAnimSequence* AnimSequence = Model->GetAnimationSequence())
 		{
 			const USkeleton* Skeleton = AnimSequence->GetSkeleton();
-		
+			if (Skeleton)
 			{
 				QUICK_SCOPE_CYCLE_COUNTER(STAT_PopulateBoneTrack);
 				for (const FBoneAnimationTrack& BoneTrack : BoneTracks)
