Ensure that the cached RawDataGUID is always cleared when PostLoading a UAnimSequencerDataModel instance. The goal of the GUID is to, when upgrading from the old data-structure to keep the GUID stable - as it can slightly vary from load-to-load without saving due to floating point differences accumulating during the upgrade path (transforms mult etc.). Once saved and thus hitting ::PostLoad we are free to clear it - and let compression run its course
#misc switch AnimSequencerController to cache the previous DataModel its GUID rather than the AnimSequence GUID - as they can get out-of-sync

diff --git a/Engine/Plugins/Animation/AnimationData/Source/AnimationData/Private/AnimSequencerController.cpp b/Engine/Plugins/Animation/AnimationData/Source/AnimationData/Private/AnimSequencerController.cpp
--- a/Engine/Plugins/Animation/AnimationData/Source/AnimationData/Private/AnimSequencerController.cpp
+++ b/Engine/Plugins/Animation/AnimationData/Source/AnimationData/Private/AnimSequencerController.cpp
@@ -2350,16 +2350,7 @@ void UAnimSequencerController::PopulateWithExistingModel(TScriptInterface<IAnima
 	if (InModel.GetObject() && InModel.GetObject()->IsA<UAnimDataModel>())
 	{
 		IAnimationDataModel::FEvaluationAndModificationLock Lock(*ModelInterface);	
-		if (const UAnimSequence* AnimationSequence = Model->GetAnimationSequence())
-		{
-			PRAGMA_DISABLE_DEPRECATION_WARNINGS
-			Model->CachedRawDataGUID = AnimationSequence->GetRawDataGuid();
-			PRAGMA_ENABLE_DEPRECATION_WARNINGS
-		}
-		else
-		{
-			ReportError(LOCTEXT("UnableToGetOuterAnimSequenceError", "Unable to retrieve outer Animation Sequence"));
-		}
+		Model->CachedRawDataGUID = InModel->GenerateGuid();
 		
 		// Copy over frame-rate from model
 		const FFrameRate ModelFrameRate = InModel->GetFrameRate();
diff --git a/Engine/Plugins/Animation/AnimationData/Source/AnimationData/Private/AnimSequencerDataModel.cpp b/Engine/Plugins/Animation/AnimationData/Source/AnimationData/Private/AnimSequencerDataModel.cpp
--- a/Engine/Plugins/Animation/AnimationData/Source/AnimationData/Private/AnimSequencerDataModel.cpp
+++ b/Engine/Plugins/Animation/AnimationData/Source/AnimationData/Private/AnimSequencerDataModel.cpp
@@ -181,6 +181,8 @@ void UAnimationSequencerDataModel::PostLoad()
 			}
 		}
 
+		CachedRawDataGUID.Invalidate();
+
 		RemoveOutOfDateControls();
 
 		ValidateData();
