Animation Curve Runtime & Editor Improvements

Runtime notes:
- Removes 'smart name' usage across the animation systems.
- Changed curve blending from a uniform array (sized per skeleton) to a sparse array of sorted named values. Blends and other combiners are performed using a dual iteration 'tape merge'.
- Skeleton curves are no longer guaranteed to cover all curve names that can be found at runtime.

Editor notes:
- Curve metadata (flags, bone links etc.) is still present on the skeleton, but can also now exist on a skeletal mesh
- Curve metadata (for morph targets) is still populated on import
- Curves can now be used arbitrarily at runtime

New features:
- New Find/Replace dialog that allows for batch-replacing curves and notifies across all of a project's assets
- New curve debugger tab in various Persona editors that allows for viewing curve values live. This also now allows viewing curves for specific pose watches.
- Pose watches now output curve tracks to the Rewind Debugger

diff --git a/Engine/Plugins/Animation/AnimationData/Source/AnimationData/Private/AnimSequencerController.cpp b/Engine/Plugins/Animation/AnimationData/Source/AnimationData/Private/AnimSequencerController.cpp
--- a/Engine/Plugins/Animation/AnimationData/Source/AnimationData/Private/AnimSequencerController.cpp
+++ b/Engine/Plugins/Animation/AnimationData/Source/AnimationData/Private/AnimSequencerController.cpp
@@ -323,139 +323,6 @@ void UAnimSequencerController::SetFrameRate(FFrameRate FrameRate, bool bShouldTr
 	}
 }
 
-void UAnimSequencerController::UpdateCurveNamesFromSkeleton(const USkeleton* Skeleton, ERawCurveTrackTypes SupportedCurveType, bool bShouldTransact /*= true*/)
-{
-	ValidateModel();
-
-	if (Skeleton)
-	{
-		if (IsSupportedCurveType(SupportedCurveType))
-		{
-			FBracket Bracket = ConditionalBracket(LOCTEXT("ValidateRawCurves", "Validating Animation Curve Names"), bShouldTransact);
-			switch (SupportedCurveType)
-			{
-			case ERawCurveTrackTypes::RCT_Float:
-				{
-					const FSmartNameMapping* NameMapping = Skeleton->GetSmartNameContainer(USkeleton::AnimCurveMappingName);
-					for (FFloatCurve& FloatCurve : Model->LegacyCurveData.FloatCurves)
-					{
-						FSmartName NewSmartName = FloatCurve.Name;
-						NameMapping->GetName(FloatCurve.Name.UID, NewSmartName.DisplayName);
-						if (NewSmartName != FloatCurve.Name)
-						{
-							const FAnimationCurveIdentifier CurrentId(FloatCurve.Name, SupportedCurveType);
-							const FAnimationCurveIdentifier NewId(NewSmartName, SupportedCurveType);
-							RenameCurve(CurrentId, NewId, bShouldTransact);
-						}
-					}
-					break;
-				}
-			case ERawCurveTrackTypes::RCT_Transform:
-				{
-					const FSmartNameMapping* NameMapping = Skeleton->GetSmartNameContainer(USkeleton::AnimTrackCurveMappingName);
-					for (FTransformCurve& TransformCurve : Model->LegacyCurveData.TransformCurves)
-					{
-						FSmartName NewSmartName = TransformCurve.Name;
-						NameMapping->GetName(TransformCurve.Name.UID, NewSmartName.DisplayName);
-						if (NewSmartName != TransformCurve.Name)
-						{
-							const FAnimationCurveIdentifier CurrentId(TransformCurve.Name, SupportedCurveType);
-							const FAnimationCurveIdentifier NewId(NewSmartName, SupportedCurveType);
-							RenameCurve(CurrentId, NewId, bShouldTransact);
-						}
-					}
-					break;
-				}
-			default:
-				{
-					checkf(false, TEXT("Unsupported curve type"));
-					break;
-				}
-			}
-		}
-		else
-		{
-			const FString CurveTypeAsString = GetCurveTypeValueName(SupportedCurveType);
-			ReportWarningf(LOCTEXT("InvalidCurveTypeWarning", "Invalid curve type provided: {0} ({1})"), FText::FromString(CurveTypeAsString), FText::AsNumber(static_cast<int32>(SupportedCurveType)));
-		}
-	}
-	else
-	{
-		ReportError(LOCTEXT("InvalidSkeletonError", "Invalid USkeleton supplied"));
-	}
-}
-
-void UAnimSequencerController::FindOrAddCurveNamesOnSkeleton(USkeleton* Skeleton, ERawCurveTrackTypes SupportedCurveType, bool bShouldTransact)
-{
-	ValidateModel();
-	
-	if (Skeleton)
-	{
-		if (IsSupportedCurveType(SupportedCurveType))
-		{
-			IAnimationDataModel::FEvaluationAndModificationLock Lock(*ModelInterface);
-			FBracket Bracket = ConditionalBracket(LOCTEXT("FindOrAddRawCurveNames", "Updating Skeleton with Animation Curve Names"), bShouldTransact);
-			switch (SupportedCurveType)
-			{
-			case ERawCurveTrackTypes::RCT_Float:
-				{
-					for (TPair<FAnimationCurveIdentifier, FAnimationCurveMetaData>& IdentifierPair : Model->CurveIdentifierToMetaData)
-					{
-						if (IdentifierPair.Key.CurveType == ERawCurveTrackTypes::RCT_Float)
-						{
-							const FSmartName CurrentSmartName = IdentifierPair.Key.InternalName;
-							FSmartName NewSmartName = CurrentSmartName;
-							Skeleton->VerifySmartName(USkeleton::AnimCurveMappingName, NewSmartName);
-							if (NewSmartName != CurrentSmartName)
-							{
-								ensure(CurrentSmartName.DisplayName == NewSmartName.DisplayName);
-								IdentifierPair.Key.InternalName = NewSmartName;
-
-								// Update float curve name itself as well
-								if (FFloatCurve* FloatCurve = Model->FindMutableFloatCurveById(FAnimationCurveIdentifier(CurrentSmartName, ERawCurveTrackTypes::RCT_Float)))
-								{
-									FloatCurve->Name = NewSmartName;
-								}
-							}
-						}
-					}
-					
-					break;
-				}
-			case ERawCurveTrackTypes::RCT_Transform:
-				{
-					for (FTransformCurve& TransformCurve : Model->LegacyCurveData.TransformCurves)
-					{
-						FSmartName NewSmartName = TransformCurve.Name;
-						Skeleton->VerifySmartName(USkeleton::AnimTrackCurveMappingName, NewSmartName);
-						if (NewSmartName != TransformCurve.Name)
-						{
-							const FAnimationCurveIdentifier CurrentId(TransformCurve.Name, SupportedCurveType);
-							const FAnimationCurveIdentifier NewId(NewSmartName, SupportedCurveType);
-							RenameCurve(CurrentId, NewId, bShouldTransact);
-						}
-					}
-					break;
-				}
-			default:
-				{
-					checkf(false, TEXT("Unsupported curve type"));
-					break;
-				}
-			}
-		}
-		else
-		{
-			const FString CurveTypeAsString = GetCurveTypeValueName(SupportedCurveType);
-			ReportWarningf(LOCTEXT("InvalidCurveTypeWarning", "Invalid curve type provided: {0} ({1})"), FText::FromString(CurveTypeAsString), FText::AsNumber(static_cast<int32>(SupportedCurveType)));
-		}
-	}
-	else
-	{
-		ReportError(LOCTEXT("InvalidSkeletonError", "Invalid USkeleton supplied"));
-	}
-}
-
 bool UAnimSequencerController::RemoveBoneTracksMissingFromSkeleton(const USkeleton* Skeleton, bool bShouldTransact /*= true*/)
 {
 	if (!ModelInterface->GetAnimationSequence())
@@ -595,8 +462,8 @@ void UAnimSequencerController::ResetModel(bool bShouldTransact /*= true*/)
 bool UAnimSequencerController::AddCurve(const FAnimationCurveIdentifier& CurveId, int32 CurveFlags /*= EAnimAssetCurveFlags::AACF_Editable*/, bool bShouldTransact /*= true*/)
 {
 	ValidateModel();
-	if (CurveId.InternalName.IsValid() || IgnoreSkeletonValidation())
-	{		
+	if (CurveId.IsValid())
+	{
 		if (IsSupportedCurveType(CurveId.CurveType))
 		{
 			if (!Model->FindCurve(CurveId))
@@ -607,7 +474,7 @@ bool UAnimSequencerController::AddCurve(const FAnimationCurveIdentifier& CurveId
 				FCurveAddedPayload Payload;
 				Payload.Identifier = CurveId;
 				
-				auto AddNewCurve = [this, CurveId, CurveName = CurveId.InternalName, CurveFlags](auto& CurveTypeArray)
+				auto AddNewCurve = [this, CurveId, CurveName = CurveId.CurveName, CurveFlags](auto& CurveTypeArray)
 				{
 					CurveTypeArray.Add({ CurveName, CurveFlags});					
 					Model->CurveIdentifierToMetaData.FindOrAdd(CurveId).Flags = CurveFlags;
@@ -620,7 +487,7 @@ bool UAnimSequencerController::AddCurve(const FAnimationCurveIdentifier& CurveId
 					break;
 				case ERawCurveTrackTypes::RCT_Float:
 					AddNewCurve(Model->LegacyCurveData.FloatCurves);
-					if(!AddCurveControl(CurveId.InternalName.DisplayName))
+					if(!AddCurveControl(CurveId.CurveName))
 					{
 						ReportError(LOCTEXT("FailedtoAddCurveControl", "Failed to add curve control"));
 					}
@@ -640,7 +507,7 @@ bool UAnimSequencerController::AddCurve(const FAnimationCurveIdentifier& CurveId
 			else
 			{
 				const FString CurveTypeAsString = GetCurveTypeValueName(CurveId.CurveType);
-				ReportWarningf(LOCTEXT("ExistingCurveNameWarning", "Curve with name {0} and type {1} ({2}) already exists"), FText::FromName(CurveId.InternalName.DisplayName), FText::FromString(CurveTypeAsString), FText::AsNumber(static_cast<int32>(CurveId.CurveType)));
+				ReportWarningf(LOCTEXT("ExistingCurveNameWarning", "Curve with name {0} and type {1} ({2}) already exists"), FText::FromName(CurveId.CurveName), FText::FromString(CurveTypeAsString), FText::AsNumber(static_cast<int32>(CurveId.CurveType)));
 			}			
 		}
 		else 
@@ -652,7 +519,7 @@ bool UAnimSequencerController::AddCurve(const FAnimationCurveIdentifier& CurveId
 	else
 	{
 		const FString CurveTypeAsString = GetCurveTypeValueName(CurveId.CurveType);
-		ReportWarningf(LOCTEXT("InvalidCurveIdentifierWarning", "Invalid curve identifier provided: name: {0}, UID: {1} type: {2}"), FText::FromName(CurveId.InternalName.DisplayName), FText::AsNumber(CurveId.InternalName.UID), FText::FromString(CurveTypeAsString));
+		ReportWarningf(LOCTEXT("InvalidCurveIdentifierWarning", "Invalid curve identifier provided: name: {0}, type: {1}"), FText::FromName(CurveId.CurveName), FText::FromString(CurveTypeAsString));
 	}
 
 	return false;
@@ -664,7 +531,7 @@ bool UAnimSequencerController::DuplicateCurve(const FAnimationCurveIdentifier& C
 
 	ERawCurveTrackTypes SupportedCurveType = CopyCurveId.CurveType;
 
-	if (CopyCurveId.InternalName.IsValid() && NewCurveId.InternalName.IsValid())
+	if (CopyCurveId.IsValid() && NewCurveId.IsValid())
 	{
 		if (IsSupportedCurveType(SupportedCurveType))
 		{
@@ -677,7 +544,7 @@ bool UAnimSequencerController::DuplicateCurve(const FAnimationCurveIdentifier& C
 						IAnimationDataModel::FEvaluationAndModificationLock Lock(*ModelInterface);
 						FTransaction Transaction = ConditionalTransaction(LOCTEXT("CopyRawCurve", "Duplicating Animation Curve"), bShouldTransact);
 
-						auto DuplicateCurve = [NewCurveId, NewCurveName = NewCurveId.InternalName, this](auto& LegacyCurveDataArray, const auto& SourceCurve)
+						auto DuplicateCurve = [NewCurveId, NewCurveName = NewCurveId.CurveName, this](auto& LegacyCurveDataArray, const auto& SourceCurve)
 						{
 							auto& DuplicatedCurve = LegacyCurveDataArray.Add_GetRef( { NewCurveName, SourceCurve.GetCurveTypeFlags() });
 							DuplicatedCurve.CopyCurve(SourceCurve);
@@ -694,7 +561,7 @@ bool UAnimSequencerController::DuplicateCurve(const FAnimationCurveIdentifier& C
 							{
 								DuplicateCurve(Model->LegacyCurveData.FloatCurves, Model->GetFloatCurve(CopyCurveId));
 															
-								if(!DuplicateCurveControl(CopyCurveId.InternalName.DisplayName, NewCurveId.InternalName.DisplayName))
+								if(!DuplicateCurveControl(CopyCurveId.CurveName, NewCurveId.CurveName))
 								{
 									ReportError(LOCTEXT("FailedToDuplicateCurveControl", "Failed to duplicate curve control"));
 								}
@@ -718,13 +585,13 @@ bool UAnimSequencerController::DuplicateCurve(const FAnimationCurveIdentifier& C
 					else
 					{
 						const FString CurveTypeAsString = GetCurveTypeValueName(NewCurveId.CurveType);
-						ReportWarningf(LOCTEXT("ExistingCurveNameWarning", "Curve with name {0} and type {1} ({2}) already exists"), FText::FromName(NewCurveId.InternalName.DisplayName), FText::FromString(CurveTypeAsString), FText::AsNumber(static_cast<int32>(NewCurveId.CurveType)));
+						ReportWarningf(LOCTEXT("ExistingCurveNameWarning", "Curve with name {0} and type {1} ({2}) already exists"), FText::FromName(NewCurveId.CurveName), FText::FromString(CurveTypeAsString), FText::AsNumber(static_cast<int32>(NewCurveId.CurveType)));
 					}
 				}
 				else
 				{
 					const FString CurveTypeAsString = GetCurveTypeValueName(CopyCurveId.CurveType);
-					ReportWarningf(LOCTEXT("CurveNameToDuplicateNotFoundWarning", "Could not find curve with name {0} and type {1} ({2}) for duplication"), FText::FromName(NewCurveId.InternalName.DisplayName), FText::FromString(CurveTypeAsString), FText::AsNumber(static_cast<int32>(NewCurveId.CurveType)));
+					ReportWarningf(LOCTEXT("CurveNameToDuplicateNotFoundWarning", "Could not find curve with name {0} and type {1} ({2}) for duplication"), FText::FromName(NewCurveId.CurveName), FText::FromString(CurveTypeAsString), FText::AsNumber(static_cast<int32>(NewCurveId.CurveType)));
 				}
 			}
 		}
@@ -734,6 +601,10 @@ bool UAnimSequencerController::DuplicateCurve(const FAnimationCurveIdentifier& C
 			ReportWarningf(LOCTEXT("InvalidCurveTypeWarning", "Invalid curve type provided: {0} ({1})"), FText::FromString(CurveTypeAsString), FText::AsNumber(static_cast<int32>(SupportedCurveType)));
 		}
 	}
+	else
+	{
+		ReportWarning(LOCTEXT("InvalidCurveIdentifierWarning", "Invalid curve identifier provided"));
+	}
 
 	return false;
 }
@@ -743,7 +614,7 @@ bool UAnimSequencerController::RemoveCurve(const FAnimationCurveIdentifier& Curv
 	ValidateModel();
 	const ERawCurveTrackTypes SupportedCurveType = CurveId.CurveType;
 
-	if (CurveId.InternalName.IsValid())
+	if (CurveId.IsValid())
 	{
 		if (IsSupportedCurveType(CurveId.CurveType))
 		{
@@ -758,7 +629,7 @@ bool UAnimSequencerController::RemoveCurve(const FAnimationCurveIdentifier& Curv
 					{
 						const FTransformCurve& TransformCurve = Model->GetTransformCurve(CurveId);
 						ConditionalAction<UE::Anim::FAddTransformCurveAction>(bShouldTransact,  CurveId, TransformCurve.GetCurveTypeFlags(), TransformCurve);
-						Model->LegacyCurveData.TransformCurves.RemoveAll([Name = TransformCurve.Name](const FTransformCurve& ToRemoveCurve) { return ToRemoveCurve.Name == Name; });
+						Model->LegacyCurveData.TransformCurves.RemoveAll([Name = TransformCurve.GetName()](const FTransformCurve& ToRemoveCurve) { return ToRemoveCurve.GetName() == Name; });
 						Model->CurveIdentifierToMetaData.Remove(CurveId);
 						break;
 					}
@@ -766,10 +637,10 @@ bool UAnimSequencerController::RemoveCurve(const FAnimationCurveIdentifier& Curv
 					{
 						const FFloatCurve& FloatCurve = Model->GetFloatCurve(CurveId);
 						ConditionalAction<UE::Anim::FAddFloatCurveAction>(bShouldTransact,  CurveId, FloatCurve.GetCurveTypeFlags(), FloatCurve.FloatCurve.GetConstRefOfKeys(), FloatCurve.Color);
-						Model->LegacyCurveData.FloatCurves.RemoveAll([Name = FloatCurve.Name](const FFloatCurve& ToRemoveCurve) { return ToRemoveCurve.Name == Name; });
+						Model->LegacyCurveData.FloatCurves.RemoveAll([Name = FloatCurve.GetName()](const FFloatCurve& ToRemoveCurve) { return ToRemoveCurve.GetName() == Name; });
 						Model->CurveIdentifierToMetaData.Remove(CurveId);
 
-						if(!RemoveCurveControl(CurveId.InternalName.DisplayName))
+						if(!RemoveCurveControl(CurveId.CurveName))
 						{
 							ReportError(LOCTEXT("FailedtoRemoveCurveControl", "Failed to remove curve control"));
 						}
@@ -791,7 +662,7 @@ bool UAnimSequencerController::RemoveCurve(const FAnimationCurveIdentifier& Curv
 			else
 			{
 				const FString CurveTypeAsString = GetCurveTypeValueName(SupportedCurveType);
-				ReportWarningf(LOCTEXT("UnableToFindCurveForRemovalWarning", "Unable to find curve: {0} of type {1}"), FText::FromName(CurveId.InternalName.DisplayName), FText::FromString(CurveTypeAsString));
+				ReportWarningf(LOCTEXT("UnableToFindCurveForRemovalWarning", "Unable to find curve: {0} of type {1}"), FText::FromName(CurveId.CurveName), FText::FromString(CurveTypeAsString));
 			}
 		}
 		else
@@ -800,7 +671,11 @@ bool UAnimSequencerController::RemoveCurve(const FAnimationCurveIdentifier& Curv
 			ReportWarningf(LOCTEXT("InvalidCurveTypeWarning", "Invalid curve type provided: {0} ({1})"), FText::FromString(CurveTypeAsString), FText::AsNumber(static_cast<int32>(CurveId.CurveType)));
 		}
 	}
-
+	else
+	{
+		const FString CurveTypeAsString = GetCurveTypeValueName(CurveId.CurveType);
+		ReportWarningf(LOCTEXT("InvalidCurveIdentifierWarning", "Invalid curve identifier provided: name: {0}, type: {1}"), FText::FromName(CurveId.CurveName), FText::FromString(CurveTypeAsString));
+	}
 	return false;
 }
 
@@ -816,7 +691,7 @@ void UAnimSequencerController::RemoveAllCurvesOfType(ERawCurveTrackTypes Support
 		TArray<FTransformCurve> TransformCurves = Model->LegacyCurveData.TransformCurves;
 		for (const FTransformCurve& Curve : TransformCurves)
 		{
-			const FAnimationCurveIdentifier CurveId(Curve.Name, ERawCurveTrackTypes::RCT_Transform);
+			const FAnimationCurveIdentifier CurveId(Curve.GetName(), ERawCurveTrackTypes::RCT_Transform);
 			RemoveCurve(CurveId, bShouldTransact);
 			Model->CurveIdentifierToMetaData.Remove(CurveId);
 		}
@@ -827,7 +702,7 @@ void UAnimSequencerController::RemoveAllCurvesOfType(ERawCurveTrackTypes Support
 		TArray<FFloatCurve> FloatCurves = Model->LegacyCurveData.FloatCurves;
 		for (const FFloatCurve& Curve : FloatCurves)
 		{
-			const FAnimationCurveIdentifier CurveId(Curve.Name, ERawCurveTrackTypes::RCT_Float);
+			const FAnimationCurveIdentifier CurveId(Curve.GetName(), ERawCurveTrackTypes::RCT_Float);
 			RemoveCurve(CurveId, bShouldTransact);
 			Model->CurveIdentifierToMetaData.Remove(CurveId);
 		}
@@ -881,7 +756,7 @@ bool UAnimSequencerController::SetCurveFlag(const FAnimationCurveIdentifier& Cur
 	else
 	{
 		const FString CurveTypeAsString = GetCurveTypeValueName(SupportedCurveType);
-		ReportWarningf(LOCTEXT("UnableToFindCurveWarning", "Unable to find curve: {0} of type {1}"), FText::FromName(CurveId.InternalName.DisplayName), FText::FromString(CurveTypeAsString));
+		ReportWarningf(LOCTEXT("UnableToFindCurveWarning", "Unable to find curve: {0} of type {1}"), FText::FromName(CurveId.CurveName), FText::FromString(CurveTypeAsString));
 	}
 
 	return false;
@@ -927,7 +802,7 @@ bool UAnimSequencerController::SetCurveFlags(const FAnimationCurveIdentifier& Cu
 	else
 	{
 		const FString CurveTypeAsString = GetCurveTypeValueName(SupportedCurveType);
-		ReportWarningf(LOCTEXT("UnableToFindCurveForRemovalWarning", "Unable to find curve: {0} of type {1}"), FText::FromName(CurveId.InternalName.DisplayName), FText::FromString(CurveTypeAsString));
+		ReportWarningf(LOCTEXT("UnableToFindCurveForRemovalWarning", "Unable to find curve: {0} of type {1}"), FText::FromName(CurveId.CurveName), FText::FromString(CurveTypeAsString));
 	}
 
 	return false;
@@ -1003,7 +878,7 @@ bool UAnimSequencerController::SetTransformCurveKeys(const FAnimationCurveIdenti
 		}
 		else
 		{
-			ReportWarningf(LOCTEXT("UnableToFindTransformCurveWarning", "Unable to find transform curve: {0}"), FText::FromName(CurveId.InternalName.DisplayName));
+			ReportWarningf(LOCTEXT("UnableToFindTransformCurveWarning", "Unable to find transform curve: {0}"), FText::FromName(CurveId.CurveName));
 		}
 	}
 	else
@@ -1065,7 +940,7 @@ bool UAnimSequencerController::SetTransformCurveKey(const FAnimationCurveIdentif
 	}
 	else
 	{
-		ReportWarningf(LOCTEXT("UnableToFindTransformCurveWarning", "Unable to find transform curve: {0}"), FText::FromName(CurveId.InternalName.DisplayName));
+		ReportWarningf(LOCTEXT("UnableToFindTransformCurveWarning", "Unable to find transform curve: {0}"), FText::FromName(CurveId.CurveName));
 	}
 
 	return false;
@@ -1077,7 +952,7 @@ bool UAnimSequencerController::RemoveTransformCurveKey(const FAnimationCurveIden
 
 	if (Model->FindMutableTransformCurveById(CurveId) != nullptr)
 	{
-		const FString BaseCurveName = CurveId.InternalName.DisplayName.ToString();
+		const FString BaseCurveName = CurveId.CurveName.ToString();
 		const TArray<FString> SubCurveNames = { TEXT( "Translation"), TEXT( "Rotation"), TEXT( "Scale") };
 		const TArray<FString> ChannelCurveNames = { TEXT("X"), TEXT("Y"), TEXT("Z") };
 
@@ -1100,7 +975,7 @@ bool UAnimSequencerController::RemoveTransformCurveKey(const FAnimationCurveIden
 	}
 	else
 	{
-		ReportWarningf(LOCTEXT("UnableToFindTransformCurveWarning", "Unable to find transform curve: {0}"), FText::FromName(CurveId.InternalName.DisplayName));
+		ReportWarningf(LOCTEXT("UnableToFindTransformCurveWarning", "Unable to find transform curve: {0}"), FText::FromName(CurveId.CurveName));
 	}
 
 	return false;
@@ -1122,20 +997,20 @@ bool UAnimSequencerController::RenameCurve(const FAnimationCurveIdentifier& Curv
 					FTransaction Transaction = ConditionalTransaction(LOCTEXT("RenameCurve", "Renaming Curve"), bShouldTransact);
 
 					FCurveRenamedPayload Payload;
-					Payload.Identifier = FAnimationCurveIdentifier(Curve->Name, CurveToRenameId.CurveType);
+					Payload.Identifier = FAnimationCurveIdentifier(Curve->GetName(), CurveToRenameId.CurveType);
 
-					Curve->Name = NewCurveId.InternalName;
+					Curve->SetName(NewCurveId.CurveName);
 					Payload.NewIdentifier = NewCurveId;
 
 					const FAnimationCurveMetaData CurveMetaData = Model->CurveIdentifierToMetaData.FindChecked(CurveToRenameId);
 					Model->CurveIdentifierToMetaData.Remove(CurveToRenameId);
 					Model->CurveIdentifierToMetaData.Add(NewCurveId, CurveMetaData);
 
-					const bool bControlNeedsRenaming = CurveToRenameId.InternalName.DisplayName != NewCurveId.InternalName.DisplayName;
+					const bool bControlNeedsRenaming = CurveToRenameId.CurveName != NewCurveId.CurveName;
 
-					if(CurveToRenameId.CurveType == ERawCurveTrackTypes::RCT_Float && (bControlNeedsRenaming && !RenameCurveControl(CurveToRenameId.InternalName.DisplayName, NewCurveId.InternalName.DisplayName)))
+					if(CurveToRenameId.CurveType == ERawCurveTrackTypes::RCT_Float && (bControlNeedsRenaming && !RenameCurveControl(CurveToRenameId.CurveName, NewCurveId.CurveName)))
 					{
-						ReportErrorf(LOCTEXT("FailedToRenameCurveControl", "Failed to rename RigCurve with name {0}"), FText::FromName(CurveToRenameId.InternalName.DisplayName));
+						ReportErrorf(LOCTEXT("FailedToRenameCurveControl", "Failed to rename RigCurve with name {0}"), FText::FromName(CurveToRenameId.CurveName));
 					}
 
 					ConditionalAction<UE::Anim::FRenameCurveAction>(bShouldTransact,  NewCurveId, CurveToRenameId);
@@ -1147,25 +1022,25 @@ bool UAnimSequencerController::RenameCurve(const FAnimationCurveIdentifier& Curv
 				else
 				{
 					const FString CurveTypeAsString = GetCurveTypeValueName(CurveToRenameId.CurveType);
-					ReportWarningf(LOCTEXT("UnableToFindCurveWarning", "Unable to find curve: {0} of type {1}"), FText::FromName(CurveToRenameId.InternalName.DisplayName), FText::FromString(CurveTypeAsString));
+					ReportWarningf(LOCTEXT("UnableToFindCurveWarning", "Unable to find curve: {0} of type {1}"), FText::FromName(CurveToRenameId.CurveName), FText::FromString(CurveTypeAsString));
 				}
 			}
 			else
 			{
 				const FString CurrentCurveTypeAsString = GetCurveTypeValueName(CurveToRenameId.CurveType);
 				const FString NewCurveTypeAsString = GetCurveTypeValueName(NewCurveId.CurveType);
-				ReportWarningf(LOCTEXT("MismatchOfCurveTypesWarning", "Different curve types provided between current and new curve names: {0} ({1}) and {2} ({3})"), FText::FromName(CurveToRenameId.InternalName.DisplayName), FText::FromString(CurrentCurveTypeAsString),
-					FText::FromName(NewCurveId.InternalName.DisplayName), FText::FromString(NewCurveTypeAsString));
+				ReportWarningf(LOCTEXT("MismatchOfCurveTypesWarning", "Different curve types provided between current and new curve names: {0} ({1}) and {2} ({3})"), FText::FromName(CurveToRenameId.CurveName), FText::FromString(CurrentCurveTypeAsString),
+					FText::FromName(NewCurveId.CurveName), FText::FromString(NewCurveTypeAsString));
 			}
 		}
 		else
 		{
-			ReportWarningf(LOCTEXT("MatchingCurveNamesWarning", "Provided curve names are the same: {0}"), FText::FromName(CurveToRenameId.InternalName.DisplayName));
+			ReportWarningf(LOCTEXT("MatchingCurveNamesWarning", "Provided curve names are the same: {0}"), FText::FromName(CurveToRenameId.CurveName));
 		}		
 	}
 	else
 	{
-		ReportWarningf(LOCTEXT("InvalidCurveIdentiferProvidedWarning", "Invalid new curve identifier provided: {2} ({3})"), FText::FromName(NewCurveId.InternalName.DisplayName), FText::AsNumber(NewCurveId.InternalName.UID));
+		ReportWarningf(LOCTEXT("InvalidCurveIdentiferProvidedWarning", "Invalid new curve identifier provided: {0}"), FText::FromName(NewCurveId.CurveName));
 	}
 
 	return false;
@@ -1197,7 +1072,7 @@ bool UAnimSequencerController::SetCurveColor(const FAnimationCurveIdentifier& Cu
 			else
 			{
 				const FString CurveTypeAsString = GetCurveTypeValueName(CurveId.CurveType);
-				ReportWarningf(LOCTEXT("UnableToFindCurveWarning", "Unable to find curve: {0} of type {1}"), FText::FromName(CurveId.InternalName.DisplayName), FText::FromString(CurveTypeAsString));
+				ReportWarningf(LOCTEXT("UnableToFindCurveWarning", "Unable to find curve: {0} of type {1}"), FText::FromName(CurveId.CurveName), FText::FromString(CurveTypeAsString));
 			}
 		}
 		else
@@ -1207,7 +1082,7 @@ bool UAnimSequencerController::SetCurveColor(const FAnimationCurveIdentifier& Cu
 	}
 	else
 	{
-		ReportWarningf(LOCTEXT("InvalidCurveIdentifier", "Invalid Curve Identifier : {0} ({1})"), FText::FromName(CurveId.InternalName.DisplayName), FText::AsNumber(CurveId.InternalName.UID));
+		ReportWarningf(LOCTEXT("InvalidCurveIdentifier", "Invalid Curve Identifier : {0}"), FText::FromName(CurveId.CurveName));
 	}	
 
 	return false;
@@ -1234,7 +1109,7 @@ bool UAnimSequencerController::ScaleCurve(const FAnimationCurveIdentifier& Curve
 			
 			ConditionalAction<UE::Anim::FScaleCurveAction>(bShouldTransact,  CurveId, Origin, 1.0f / Factor, SupportedCurveType);
 
-			if (!SetCurveControlKeys(CurveId.InternalName.DisplayName, Curve->FloatCurve.GetConstRefOfKeys()))
+			if (!SetCurveControlKeys(CurveId.CurveName, Curve->FloatCurve.GetConstRefOfKeys()))
 			{
 				ReportError(LOCTEXT("FailedToScaleControlCurve", "Failed to scale curve control"));
 			}
@@ -1245,7 +1120,7 @@ bool UAnimSequencerController::ScaleCurve(const FAnimationCurveIdentifier& Curve
 		}
 		else
 		{
-			ReportWarningf(LOCTEXT("UnableToFindFloatCurveWarning", "Unable to find float curve: {0}"), FText::FromName(CurveId.InternalName.DisplayName));
+			ReportWarningf(LOCTEXT("UnableToFindFloatCurveWarning", "Unable to find float curve: {0}"), FText::FromName(CurveId.CurveName));
 		}
 	}
 	else
@@ -1267,7 +1142,7 @@ bool UAnimSequencerController::SetCurveKey(const FAnimationCurveIdentifier& Curv
 
 		auto SetControlCurve = [this, CurveId, Key]()
 		{
-			if(CurveId.CurveType == ERawCurveTrackTypes::RCT_Float && !SetCurveControlKey(CurveId.InternalName.DisplayName, Key))
+			if(CurveId.CurveType == ERawCurveTrackTypes::RCT_Float && !SetCurveControlKey(CurveId.CurveName, Key))
 			{
 				ReportError(LOCTEXT("FailedtoSetCurveControlKey", "Failed to set curve control key"));
 			}	
@@ -1307,7 +1182,7 @@ bool UAnimSequencerController::SetCurveKey(const FAnimationCurveIdentifier& Curv
 	}
 	else
 	{
-		ReportErrorf(LOCTEXT("RichCurveNotFoundError", "Unable to find rich curve: curve name {0}"), FText::FromName(CurveId.InternalName.DisplayName));
+		ReportErrorf(LOCTEXT("RichCurveNotFoundError", "Unable to find rich curve: curve name {0}"), FText::FromName(CurveId.CurveName));
 	}
 
 	return false;
@@ -1335,7 +1210,7 @@ bool UAnimSequencerController::RemoveCurveKey(const FAnimationCurveIdentifier& C
 
 			RichCurve->DeleteKey(Handle);
 			
-			if(CurveId.CurveType == ERawCurveTrackTypes::RCT_Float && !RemoveCurveControlKey(CurveId.InternalName.DisplayName, Time))
+			if(CurveId.CurveType == ERawCurveTrackTypes::RCT_Float && !RemoveCurveControlKey(CurveId.CurveName, Time))
 			{
 				ReportError(LOCTEXT("FailedtoRemoveCurveControlKeys", "Failed to remove curve control key"));
 			}
@@ -1346,12 +1221,12 @@ bool UAnimSequencerController::RemoveCurveKey(const FAnimationCurveIdentifier& C
 		}
 		else
 		{
-			ReportErrorf(LOCTEXT("RichCurveKeyNotFoundError", "Unable to find rich curve key: curve name {0}, time {1}"), FText::FromName(CurveId.InternalName.DisplayName), FText::AsNumber(Time));
+			ReportErrorf(LOCTEXT("RichCurveKeyNotFoundError", "Unable to find rich curve key: curve name {0}, time {1}"), FText::FromName(CurveId.CurveName), FText::AsNumber(Time));
 		}
 	}
 	else
 	{
-		ReportErrorf(LOCTEXT("RichCurveNotFoundError", "Unable to find rich curve: curve name {0}"), FText::FromName(CurveId.InternalName.DisplayName));
+		ReportErrorf(LOCTEXT("RichCurveNotFoundError", "Unable to find rich curve: curve name {0}"), FText::FromName(CurveId.CurveName));
 	}
 
 	return false;
@@ -1377,12 +1252,12 @@ bool UAnimSequencerController::SetCurveKeys(const FAnimationCurveIdentifier& Cur
 		}
 		else
 		{
-			ReportErrorf(LOCTEXT("RichCurveNotFoundError", "Unable to find rich curve: curve name {0}"), FText::FromName(CurveId.InternalName.DisplayName));
+			ReportErrorf(LOCTEXT("RichCurveNotFoundError", "Unable to find rich curve: curve name {0}"), FText::FromName(CurveId.CurveName));
 		}
 
 		if(CurveId.CurveType == ERawCurveTrackTypes::RCT_Float)
 		{
-			if (SetCurveControlKeys(CurveId.InternalName.DisplayName, CurveKeys))
+			if (SetCurveControlKeys(CurveId.CurveName, CurveKeys))
 			{
 				bKeysSet = true;
 			}
@@ -1403,7 +1278,7 @@ bool UAnimSequencerController::SetCurveKeys(const FAnimationCurveIdentifier& Cur
 	}
 	// else
 	// {
-	// 	ReportErrorf(LOCTEXT("CurveNotFoundError", "Unable to find curve: curve name {0}"), FText::FromName(CurveId.InternalName.DisplayName));
+	// 	ReportErrorf(LOCTEXT("CurveNotFoundError", "Unable to find curve: curve name {0}"), FText::FromName(CurveId.CurveName));
 	// }	
 
 	return false;
@@ -1432,7 +1307,7 @@ bool UAnimSequencerController::SetCurveAttributes(const FAnimationCurveIdentifie
 				const URigHierarchy* Hierarchy = FKRig->GetHierarchy();
 				if (Hierarchy ||  IgnoreSkeletonValidation())
 				{
-					const FRigElementKey CurveKey(URigHierarchy::GetSanitizedName(CurveId.InternalName.DisplayName.ToString()), ERigElementType::Curve);
+					const FRigElementKey CurveKey(URigHierarchy::GetSanitizedName(CurveId.CurveName.ToString()), ERigElementType::Curve);
 					const FRigElementKey CurveControlKey(UFKControlRig::GetControlName(CurveKey.Name, ERigElementType::Curve), ERigElementType::Control);
 
 					const bool bContainsCurve = Hierarchy ? Hierarchy->Contains(CurveKey) : IgnoreSkeletonValidation();
@@ -1868,7 +1743,7 @@ void UAnimSequencerController::ResizeCurves(float NewLength, bool bInserted, flo
 	{
 		FFloatCurve ResizedCurve = Curve;
 		ResizedCurve.Resize(NewLength, bInserted, T0, T1);
-		SetCurveKeys(FAnimationCurveIdentifier(Curve.Name, ERawCurveTrackTypes::RCT_Float), ResizedCurve.FloatCurve.GetConstRefOfKeys(), bShouldTransact);
+		SetCurveKeys(FAnimationCurveIdentifier(Curve.GetName(), ERawCurveTrackTypes::RCT_Float), ResizedCurve.FloatCurve.GetConstRefOfKeys(), bShouldTransact);
 	}
 
 	for (FTransformCurve& Curve : Model->LegacyCurveData.TransformCurves)
@@ -1881,7 +1756,7 @@ void UAnimSequencerController::ResizeCurves(float NewLength, bool bInserted, flo
 			for (int32 ChannelIndex = 0; ChannelIndex < 3; ++ChannelIndex)
 			{
 				const EVectorCurveChannel Axis = static_cast<EVectorCurveChannel>(ChannelIndex);
-				FAnimationCurveIdentifier TargetCurveIdentifier = FAnimationCurveIdentifier(Curve.Name, ERawCurveTrackTypes::RCT_Transform);
+				FAnimationCurveIdentifier TargetCurveIdentifier = FAnimationCurveIdentifier(Curve.GetName(), ERawCurveTrackTypes::RCT_Transform);
 				UAnimationCurveIdentifierExtensions::GetTransformChildCurveIdentifier(TargetCurveIdentifier, Channel, Axis);
 				
 				FRichCurve& ChannelCurve = SubCurve.FloatCurves[ChannelIndex];
@@ -2301,7 +2176,7 @@ void UAnimSequencerController::RemoveUnusedControlsAndCurves() const
 				{
 					if(!FloatCurves.ContainsByPredicate([CurveElement](const FFloatCurve& FloatCurve)
 					{
-						return FloatCurve.Name.DisplayName == CurveElement->GetName();
+						return FloatCurve.GetName() == CurveElement->GetName();
 					}))
 					{						
 						ElementsToRemove.Add(CurveElement->GetKey());
@@ -2323,10 +2198,6 @@ void UAnimSequencerController::UpdateWithSkeleton(USkeleton* TargetSkeleton, boo
 {
 	OpenBracket(LOCTEXT("SettingNewskeleton", "Updating Skeleton for Animation Data Model"));
 	{
-		// Update the curve names with the new skeleton
-		FindOrAddCurveNamesOnSkeleton(TargetSkeleton, ERawCurveTrackTypes::RCT_Float, bShouldTransact);	
-		FindOrAddCurveNamesOnSkeleton(TargetSkeleton, ERawCurveTrackTypes::RCT_Transform, bShouldTransact);
-
 		// (re-)generate the rig hierarchy
 		Model->InitializeFKControlRig(CastChecked<UFKControlRig>(Model->GetControlRig()), TargetSkeleton);
 		
@@ -2407,7 +2278,7 @@ void UAnimSequencerController::PopulateWithExistingModel(TScriptInterface<IAnima
 				QUICK_SCOPE_CYCLE_COUNTER(STAT_PopulateCurves);
 				for (const FFloatCurve& Curve : InModel->GetCurveData().FloatCurves)
 				{
-					const FAnimationCurveIdentifier CurveId = Skeleton ? UAnimationCurveIdentifierExtensions::FindCurveIdentifier(Skeleton, Curve.Name.DisplayName, ERawCurveTrackTypes::RCT_Float) : FAnimationCurveIdentifier(Curve.Name, ERawCurveTrackTypes::RCT_Float);
+					const FAnimationCurveIdentifier CurveId = Skeleton ? UAnimationCurveIdentifierExtensions::FindCurveIdentifier(Skeleton, Curve.GetName(), ERawCurveTrackTypes::RCT_Float) : FAnimationCurveIdentifier(Curve.GetName(), ERawCurveTrackTypes::RCT_Float);
 					if (CurveId.IsValid() || IgnoreSkeletonValidation())
 					{
 						AddCurve(CurveId, Curve.GetCurveTypeFlags(), false);
@@ -2421,7 +2292,7 @@ void UAnimSequencerController::PopulateWithExistingModel(TScriptInterface<IAnima
 				
 				for (const FTransformCurve& TransformCurve : InModel->GetCurveData().TransformCurves)
 				{
-					const FAnimationCurveIdentifier CurveId = Skeleton ? UAnimationCurveIdentifierExtensions::FindCurveIdentifier(Skeleton, TransformCurve.Name.DisplayName, ERawCurveTrackTypes::RCT_Transform) : FAnimationCurveIdentifier(TransformCurve.Name, ERawCurveTrackTypes::RCT_Transform);
+					const FAnimationCurveIdentifier CurveId = Skeleton ? UAnimationCurveIdentifierExtensions::FindCurveIdentifier(Skeleton, TransformCurve.GetName(), ERawCurveTrackTypes::RCT_Transform) : FAnimationCurveIdentifier(TransformCurve.GetName(), ERawCurveTrackTypes::RCT_Transform);
 					if (CurveId.IsValid() || IgnoreSkeletonValidation())
 					{
 						AddCurve(CurveId, TransformCurve.GetCurveTypeFlags());
@@ -2442,7 +2313,7 @@ void UAnimSequencerController::PopulateWithExistingModel(TScriptInterface<IAnima
 					}
 					else
 					{
-						ReportWarningf(LOCTEXT("InvalidCurveIdentifierFoundWarning", "Invalid transform curve identifier found: name: {0}, UID: {1} type: {2}"), FText::FromName(CurveId.InternalName.DisplayName), FText::AsNumber(CurveId.InternalName.UID));
+						ReportWarningf(LOCTEXT("InvalidCurveIdentifierFoundWarning", "Invalid transform curve identifier found: name: {0}, type: {1}"), FText::FromName(CurveId.CurveName), FText::AsNumber(static_cast<int32>(CurveId.CurveType)));
 					}
 				}
 			}
diff --git a/Engine/Plugins/Animation/AnimationData/Source/AnimationData/Private/AnimSequencerController.h b/Engine/Plugins/Animation/AnimationData/Source/AnimationData/Private/AnimSequencerController.h
--- a/Engine/Plugins/Animation/AnimationData/Source/AnimationData/Private/AnimSequencerController.h
+++ b/Engine/Plugins/Animation/AnimationData/Source/AnimationData/Private/AnimSequencerController.h
@@ -63,8 +63,6 @@ public:
 	virtual bool RemoveCurveKey(const FAnimationCurveIdentifier& CurveId, float Time, bool bShouldTransact = true) override;
 	virtual bool SetCurveKeys(const FAnimationCurveIdentifier& CurveId, const TArray<FRichCurveKey>& CurveKeys, bool bShouldTransact = true) override;
 	virtual bool SetCurveAttributes(const FAnimationCurveIdentifier& CurveId, const FCurveAttributes& Attributes, bool bShouldTransact = true) override;
-	virtual void UpdateCurveNamesFromSkeleton(const USkeleton* Skeleton, ERawCurveTrackTypes SupportedCurveType, bool bShouldTransact = true) override;
-	virtual void FindOrAddCurveNamesOnSkeleton(USkeleton* Skeleton, ERawCurveTrackTypes SupportedCurveType, bool bShouldTransact = true) override;
 	virtual bool RemoveBoneTracksMissingFromSkeleton(const USkeleton* Skeleton, bool bShouldTransact = true) override;
 	virtual void UpdateAttributesFromSkeleton(const USkeleton* Skeleton, bool bShouldTransact = true) override;
 	virtual void NotifyPopulated() override;
diff --git a/Engine/Plugins/Animation/AnimationData/Source/AnimationData/Private/AnimSequencerDataModel.cpp b/Engine/Plugins/Animation/AnimationData/Source/AnimationData/Private/AnimSequencerDataModel.cpp
--- a/Engine/Plugins/Animation/AnimationData/Source/AnimationData/Private/AnimSequencerDataModel.cpp
+++ b/Engine/Plugins/Animation/AnimationData/Source/AnimationData/Private/AnimSequencerDataModel.cpp
@@ -22,6 +22,7 @@
 #include "Animation/AnimationSettings.h"
 #include "Compilation/MovieSceneCompiledDataManager.h"
 #include "UObject/UObjectThreadContext.h"
+#include "Animation/AnimCurveTypes.h"
 
 #include "Runtime/MovieScene/Private/Channels/MovieSceneCurveChannelImpl.h"
 
@@ -72,7 +73,7 @@ void UAnimationSequencerDataModel::RemoveOutOfDateControls() const
 					Hierarchy->ForEach<FRigCurveElement>([this, &ElementKeysToRemove](const FRigCurveElement* CurveElement) -> bool
 					{
 						const FName TargetCurveName = CurveElement->GetName();
-						if(!LegacyCurveData.FloatCurves.ContainsByPredicate([TargetCurveName](const FFloatCurve& Curve) { return Curve.Name.DisplayName == TargetCurveName; }))
+						if(!LegacyCurveData.FloatCurves.ContainsByPredicate([TargetCurveName](const FFloatCurve& Curve) { return Curve.GetName() == TargetCurveName; }))
 						{
 							ElementKeysToRemove.Add(CurveElement->GetKey());	
 						}
@@ -367,7 +368,7 @@ const FFloatCurve* UAnimationSequencerDataModel::FindFloatCurve(const FAnimation
 	ensure(CurveIdentifier.CurveType == ERawCurveTrackTypes::RCT_Float);
 	for (const FFloatCurve& FloatCurve : GetCurveData().FloatCurves)
 	{
-		if (FloatCurve.Name == CurveIdentifier.InternalName || (FloatCurve.Name.UID == CurveIdentifier.InternalName.UID && FloatCurve.Name.UID != SmartName::MaxUID))
+		if (FloatCurve.GetName() == CurveIdentifier.CurveName)
 		{
 			return &FloatCurve;
 		}
@@ -381,7 +382,7 @@ const FTransformCurve* UAnimationSequencerDataModel::FindTransformCurve(const FA
 	ensure(CurveIdentifier.CurveType == ERawCurveTrackTypes::RCT_Transform);
     for (const FTransformCurve& TransformCurve : GetCurveData().TransformCurves)
     {
-    	if (TransformCurve.Name == CurveIdentifier.InternalName || TransformCurve.Name.UID == CurveIdentifier.InternalName.UID)
+    	if (TransformCurve.GetName() == CurveIdentifier.CurveName)
     	{
     		return &TransformCurve;
     	}
@@ -650,7 +651,7 @@ FGuid UAnimationSequencerDataModel::GenerateGuid() const
 
 		for (const FTransformCurve& Curve : GetTransformCurves())
 		{
-			UpdateWithData(Curve.Name.UID);
+			UpdateWithData(Curve.GetName());
 
 			auto UpdateWithComponent = [&UpdateWithFloatCurve](const FVectorCurve& VectorCurve)
 			{
@@ -823,69 +824,65 @@ void UAnimationSequencerDataModel::GenerateLegacyCurveData()
 			{
 				if (const UControlRig* ControlRig = Section->GetControlRig())
 				{
-					if (USkeleton* Skeleton = GetSkeleton())
+					if(URigHierarchy* Hierarchy = ControlRig->GetHierarchy())
 					{
-						if(URigHierarchy* Hierarchy = ControlRig->GetHierarchy())
+						const FString SequencerSuffix(TEXT("_Sequencer"));
+						const TArray<FScalarParameterNameAndCurve>& ScalarCurves = Section->GetScalarParameterNamesAndCurves();				
+						if (RetainFloatCurves)
 						{
-							const FString SequencerSuffix(TEXT("_Sequencer"));
-							const TArray<FScalarParameterNameAndCurve>& ScalarCurves = Section->GetScalarParameterNamesAndCurves();				
-							if (RetainFloatCurves)
+							LegacyCurveData.FloatCurves.RemoveAll([SequencerSuffix](const FFloatCurve& FloatCurve)
 							{
-								LegacyCurveData.FloatCurves.RemoveAll([SequencerSuffix](const FFloatCurve& FloatCurve)
-								{
-									return FloatCurve.Name.DisplayName.ToString().EndsWith(SequencerSuffix);
-								});
-							}
-							else
-							{
-								LegacyCurveData.FloatCurves.Empty();
-							}				
+								return FloatCurve.GetName().ToString().EndsWith(SequencerSuffix);
+							});
+						}
+						else
+						{
+							LegacyCurveData.FloatCurves.Empty();
+						}				
 
-							Hierarchy->ForEach<FRigCurveElement>([Hierarchy,Skeleton, this, ScalarCurves, FrameRate = GetFrameRate(), SequencerSuffix](const FRigCurveElement* CurveElement) -> bool
+						Hierarchy->ForEach<FRigCurveElement>([Hierarchy, this, ScalarCurves, FrameRate = GetFrameRate(), SequencerSuffix](const FRigCurveElement* CurveElement) -> bool
+						{
+							const FRigElementKey ControlKey(UFKControlRig::GetControlName(CurveElement->GetName(), ERigElementType::Curve), ERigElementType::Control);
+							if (const FRigControlElement* Element = Hierarchy->Find<FRigControlElement>(ControlKey))
 							{
-								const FRigElementKey ControlKey(UFKControlRig::GetControlName(CurveElement->GetName(), ERigElementType::Curve), ERigElementType::Control);
-								if (const FRigControlElement* Element = Hierarchy->Find<FRigControlElement>(ControlKey))
+								FFloatCurve& FloatCurve = LegacyCurveData.FloatCurves.AddDefaulted_GetRef();
+								if (RetainFloatCurves)
 								{
-									FFloatCurve& FloatCurve = LegacyCurveData.FloatCurves.AddDefaulted_GetRef();
-									if (RetainFloatCurves)
-									{
-										FloatCurve.Name.DisplayName = FName(*(CurveElement->GetName().ToString() + TEXT("_Sequencer")));
-									}
-									else
-									{
-										FloatCurve.Name.DisplayName = CurveElement->GetName();	
-									}						
+									FloatCurve.SetName(FName(*(CurveElement->GetName().ToString() + TEXT("_Sequencer"))));
+								}
+								else
+								{
+									FloatCurve.SetName(CurveElement->GetName());
+								}						
+
+								FloatCurve.Color = Element->Settings.ShapeColor;
 								
-									Skeleton->VerifySmartName(USkeleton::AnimCurveMappingName, FloatCurve.Name);
-									FloatCurve.Color = Element->Settings.ShapeColor;
-									
-									const FAnimationCurveIdentifier CurveId(FloatCurve.Name, ERawCurveTrackTypes::RCT_Float);
-									if (!RetainFloatCurves || !FloatCurve.Name.DisplayName.ToString().Contains(SequencerSuffix))
-									{
-										if (CurveIdentifierToMetaData.Contains(CurveId))
-										{
-											const FAnimationCurveMetaData& CurveMetaData = CurveIdentifierToMetaData.FindChecked(CurveId);
-											FloatCurve.SetCurveTypeFlags(CurveMetaData.Flags);
-											FloatCurve.Color = CurveMetaData.Color;
-										}
-									}							
-
-									if (const FScalarParameterNameAndCurve* ScalarCurve = ScalarCurves.FindByPredicate([Element](FScalarParameterNameAndCurve Curve)
+								const FAnimationCurveIdentifier CurveId(FloatCurve.GetName(), ERawCurveTrackTypes::RCT_Float);
+								if (!RetainFloatCurves || !FloatCurve.GetName().ToString().Contains(SequencerSuffix))
+								{
+									if (CurveIdentifierToMetaData.Contains(CurveId))
 									{
-										return Curve.ParameterName == Element->GetName();
-									}))
-									{							
-										AnimSequencerHelpers::ConvertFloatChannelToRichCurve(ScalarCurve->ParameterCurve, FloatCurve.FloatCurve, FrameRate);
+										const FAnimationCurveMetaData& CurveMetaData = CurveIdentifierToMetaData.FindChecked(CurveId);
+										FloatCurve.SetCurveTypeFlags(CurveMetaData.Flags);
+										FloatCurve.Color = CurveMetaData.Color;
 									}
+								}							
+
+								if (const FScalarParameterNameAndCurve* ScalarCurve = ScalarCurves.FindByPredicate([Element](FScalarParameterNameAndCurve Curve)
+								{
+									return Curve.ParameterName == Element->GetName();
+								}))
+								{							
+									AnimSequencerHelpers::ConvertFloatChannelToRichCurve(ScalarCurve->ParameterCurve, FloatCurve.FloatCurve, FrameRate);
 								}
-								return true;
-							});	
-						}
-						else
-						{						
-							IAnimationDataController::ReportObjectErrorf(this, LOCTEXT("UnableToFindRigHierarchy", "Unable to retrieve RigHierarchy for ControlRig ({0})"), FText::FromString(ControlRig->GetPathName()));	      
-						}
-					}								
+							}
+							return true;
+						});	
+					}
+					else
+					{						
+						IAnimationDataController::ReportObjectErrorf(this, LOCTEXT("UnableToFindRigHierarchy", "Unable to retrieve RigHierarchy for ControlRig ({0})"), FText::FromString(ControlRig->GetPathName()));	      
+					}
 				}
 			}
 		}
@@ -980,7 +977,7 @@ void UAnimationSequencerDataModel::ValidateLegacyAgainstControlRigData() const
 		// Validate curve data against controls
 		for (const FFloatCurve& FloatCurve : LegacyCurveData.FloatCurves)
 		{
-			const FName CurveName = FloatCurve.Name.DisplayName;					
+			const FName CurveName = FloatCurve.GetName();					
 			const FRigElementKey CurveKey(CurveName, ERigElementType::Curve);
 			const FRigCurveElement* CurveElement = Hierarchy->Find<FRigCurveElement>(CurveKey);
 			if (!CurveElement)
@@ -1131,8 +1128,10 @@ void UAnimationSequencerDataModel::GeneratePoseData(UControlRig* ControlRig, FAn
 			RigPose.ResetToRefPose();
 			const FBoneContainer& RequiredBones = RigPose.GetBoneContainer();
 			FBlendedCurve& Curve = InOutPoseData.GetCurve();
+			Curve.Empty();
+
 			UE::Anim::Retargeting::FRetargetingScope RetargetingScope(GetSkeleton(), RigPose, EvaluationContext);
-			
+
 			// Populate bone/curve elements to Pose/Curve indices
 			{
 				QUICK_SCOPE_CYCLE_COUNTER(STAT_GetMappings);
@@ -1159,27 +1158,13 @@ void UAnimationSequencerDataModel::GeneratePoseData(UControlRig* ControlRig, FAn
 					return true;
 				});
 
-				if (Curve.IsValid())
+				RigHierarchy->ForEach<FRigCurveElement>([this, &RigHierarchy, &Curve](const FRigCurveElement* CurveElement) -> bool
 				{
-					// Called during compression that can occur while GC is in progress, marking weakptrs as unreachable temporarily
-					const FSmartNameMapping* SmartNameMapping = RequiredBones.GetSkeletonAsset(true)->GetSmartNameContainer(USkeleton::AnimCurveMappingName);
-					RigHierarchy->ForEach<FRigCurveElement>([this, &RigHierarchy, &Curve, SmartNameMapping](const FRigCurveElement* CurveElement) -> bool
-					{
-						const FName& CurveName = CurveElement->GetName();
-						const SmartName::UID_Type CurveIndex = SmartNameMapping->FindUID(CurveName);
-						if (CurveIndex != INDEX_NONE)
-						{
-							if (Curve.IsEnabled(CurveIndex))
-							{
-								Curve.Set(CurveIndex, RigHierarchy->GetCurveValue(CurveElement->GetKey()));
-							}
-							
-						}
-						return true;
-					});
-				}
-
-			}			
+					const FName& CurveName = CurveElement->GetName();
+					Curve.Add(CurveName, RigHierarchy->GetCurveValue(CurveElement->GetKey()));
+					return true;
+				});
+			}
 
 			{
 				QUICK_SCOPE_CYCLE_COUNTER(STAT_NormalizeRotations);
@@ -1198,7 +1183,7 @@ void UAnimationSequencerDataModel::GeneratePoseData(UControlRig* ControlRig, FAn
 					}
 			
 					// Add or retrieve curve
-					const FName& CurveName = TransformCurve.Name.DisplayName;
+					const FName& CurveName = TransformCurve.GetName();
 					// note we're not checking Curve.GetCurveTypeFlags() yet
 					FTransform Value = TransformCurve.Evaluate(EvaluationContext.SampleFrameRate.AsSeconds(EvaluationContext.SampleTime), 1.f);
 
@@ -1444,7 +1429,7 @@ FTransformCurve* UAnimationSequencerDataModel::FindMutableTransformCurveById(con
 {
 	for (FTransformCurve& TransformCurve : LegacyCurveData.TransformCurves)
 	{
-		if (TransformCurve.Name.UID == CurveIdentifier.InternalName.UID)
+		if (TransformCurve.GetName() == CurveIdentifier.CurveName)
 		{
 			return &TransformCurve;
 		}
@@ -1457,7 +1442,7 @@ FFloatCurve* UAnimationSequencerDataModel::FindMutableFloatCurveById(const FAnim
 {
 	for (FFloatCurve& FloatCurve : LegacyCurveData.FloatCurves)
 	{
-		if (FloatCurve.Name.UID == CurveIdentifier.InternalName.UID)
+		if (FloatCurve.GetName() == CurveIdentifier.CurveName)
 		{
 			return &FloatCurve;
 		}
diff --git a/Engine/Plugins/Animation/AnimationLocomotionLibrary/Source/Runtime/Private/AnimDistanceMatchingLibrary.cpp b/Engine/Plugins/Animation/AnimationLocomotionLibrary/Source/Runtime/Private/AnimDistanceMatchingLibrary.cpp
--- a/Engine/Plugins/Animation/AnimationLocomotionLibrary/Source/Runtime/Private/AnimDistanceMatchingLibrary.cpp
+++ b/Engine/Plugins/Animation/AnimationLocomotionLibrary/Source/Runtime/Private/AnimDistanceMatchingLibrary.cpp
@@ -14,9 +14,9 @@ DEFINE_LOG_CATEGORY_STATIC(LogAnimDistanceMatchingLibrary, Verbose, All);
 
 namespace UE::Anim::DistanceMatchingUtility
 {
-	static float GetDistanceRange(const UAnimSequenceBase* InAnimSequence, USkeleton::AnimCurveUID CurveUID)
+	static float GetDistanceRange(const UAnimSequenceBase* InAnimSequence, FName CurveName)
 	{
-		FAnimCurveBufferAccess BufferCurveAccess(InAnimSequence, CurveUID);
+		FAnimCurveBufferAccess BufferCurveAccess(InAnimSequence, CurveName);
 		if (BufferCurveAccess.IsValid())
 		{
 			const int32 NumSamples = BufferCurveAccess.GetNumSamples();
@@ -28,23 +28,9 @@ namespace UE::Anim::DistanceMatchingUtility
 		return 0.f;
 	}
 
-	static USkeleton::AnimCurveUID GetCurveUID(const UAnimSequenceBase* InAnimSequence, FName CurveName)
-	{
-		if (const USkeleton* Skeleton = InAnimSequence->GetSkeleton())
-		{
-			const FSmartNameMapping* CurveNameMapping = Skeleton->GetSmartNameContainer(USkeleton::AnimCurveMappingName);
-			if (CurveNameMapping)
-			{
-				return CurveNameMapping->FindUID(CurveName);
-			}
-		}
-
-		return SmartName::MaxUID;
-	}
-
-	static float GetAnimPositionFromDistance(const UAnimSequenceBase* InAnimSequence, const float& InDistance, USkeleton::AnimCurveUID CurveUID)
+	static float GetAnimPositionFromDistance(const UAnimSequenceBase* InAnimSequence, const float& InDistance, FName InCurveName)
 	{	
-		FAnimCurveBufferAccess BufferCurveAccess(InAnimSequence, CurveUID);
+		FAnimCurveBufferAccess BufferCurveAccess(InAnimSequence, InCurveName);
 		if (BufferCurveAccess.IsValid())
 		{
 			const int32 NumKeys = BufferCurveAccess.GetNumSamples();
@@ -93,13 +79,13 @@ namespace UE::Anim::DistanceMatchingUtility
 	/**
 	* Advance from the current time to a new time in the animation that will result in the desired distance traveled by the authored root motion.
 	*/
-	static float GetTimeAfterDistanceTraveled(const UAnimSequenceBase* AnimSequence, float CurrentTime, float DistanceTraveled, USkeleton::AnimCurveUID CurveUID, const bool bAllowLooping)
+	static float GetTimeAfterDistanceTraveled(const UAnimSequenceBase* AnimSequence, float CurrentTime, float DistanceTraveled, FName CurveName, const bool bAllowLooping)
 	{
 		float NewTime = CurrentTime;
 		if (AnimSequence != nullptr)
 		{
 			// Avoid infinite loops if the animation doesn't cover any distance.
-			if (!FMath::IsNearlyZero(UE::Anim::DistanceMatchingUtility::GetDistanceRange(AnimSequence, CurveUID)))
+			if (!FMath::IsNearlyZero(UE::Anim::DistanceMatchingUtility::GetDistanceRange(AnimSequence, CurveName)))
 			{
 				float AccumulatedDistance = 0.f;
 
@@ -114,8 +100,8 @@ namespace UE::Anim::DistanceMatchingUtility
 				// Traverse the distance curve, accumulating animated distance until the desired distance is reached.
 				while ((AccumulatedDistance < DistanceTraveled) && (bAllowLooping || (NewTime + StepTime < SequenceLength)))
 				{
-					const float CurrentDistance = AnimSequence->EvaluateCurveData(CurveUID, NewTime);
-					const float DistanceAfterStep = AnimSequence->EvaluateCurveData(CurveUID, NewTime + StepTime);
+					const float CurrentDistance = AnimSequence->EvaluateCurveData(CurveName, NewTime);
+					const float DistanceAfterStep = AnimSequence->EvaluateCurveData(CurveName, NewTime + StepTime);
 					const float AnimationDistanceThisStep = DistanceAfterStep - CurrentDistance;
 
 					if (!FMath::IsNearlyZero(AnimationDistanceThisStep))
@@ -181,8 +167,7 @@ FSequenceEvaluatorReference UAnimDistanceMatchingLibrary::AdvanceTimeByDistanceM
 						const float CurrentAssetLength = InSequenceEvaluator.GetCurrentAssetLength();
 						const bool bAllowLooping = InSequenceEvaluator.GetShouldLoop();
 
-						const USkeleton::AnimCurveUID CurveUID = UE::Anim::DistanceMatchingUtility::GetCurveUID(AnimSequence, DistanceCurveName);
-						float TimeAfterDistanceTraveled = UE::Anim::DistanceMatchingUtility::GetTimeAfterDistanceTraveled(AnimSequence, CurrentTime, DistanceTraveled, CurveUID, bAllowLooping);
+						float TimeAfterDistanceTraveled = UE::Anim::DistanceMatchingUtility::GetTimeAfterDistanceTraveled(AnimSequence, CurrentTime, DistanceTraveled, DistanceCurveName, bAllowLooping);
 
 						// Calculate the effective playrate that would result from advancing the animation by the distance traveled.
 						// Account for the animation looping.
@@ -231,19 +216,11 @@ FSequenceEvaluatorReference UAnimDistanceMatchingLibrary::DistanceMatchToTarget(
 		{
 			if (const UAnimSequenceBase* AnimSequence = Cast<UAnimSequence>(InSequenceEvaluator.GetSequence()))
 			{
-				const USkeleton::AnimCurveUID CurveUID = UE::Anim::DistanceMatchingUtility::GetCurveUID(AnimSequence, DistanceCurveName);
-				if (AnimSequence->HasCurveData(CurveUID))
-				{
-					// By convention, distance curves store the distance to a target as a negative value.
-					const float NewTime = UE::Anim::DistanceMatchingUtility::GetAnimPositionFromDistance(AnimSequence, -DistanceToTarget, CurveUID);
-					if (!InSequenceEvaluator.SetExplicitTime(NewTime))
-					{
-						UE_LOG(LogAnimDistanceMatchingLibrary, Warning, TEXT("Could not set explicit time on sequence evaluator, value is not dynamic. Set it as Always Dynamic."));
-					}
-				}
-				else
+				// By convention, distance curves store the distance to a target as a negative value.
+				const float NewTime = UE::Anim::DistanceMatchingUtility::GetAnimPositionFromDistance(AnimSequence, -DistanceToTarget, DistanceCurveName);
+				if (!InSequenceEvaluator.SetExplicitTime(NewTime))
 				{
-					UE_LOG(LogAnimDistanceMatchingLibrary, Warning, TEXT("DistanceMatchToTarget called with invalid DistanceCurveName or animation (%s) is missing a distance curve."), *GetNameSafe(AnimSequence));
+					UE_LOG(LogAnimDistanceMatchingLibrary, Warning, TEXT("Could not set explicit time on sequence evaluator, value is not dynamic. Set it as Always Dynamic."));
 				}
 			}
 			else
diff --git a/Engine/Plugins/Animation/AnimationModifierLibrary/Source/AnimationModifierLibrary/Private/MotionExtractorModifier.cpp b/Engine/Plugins/Animation/AnimationModifierLibrary/Source/AnimationModifierLibrary/Private/MotionExtractorModifier.cpp
--- a/Engine/Plugins/Animation/AnimationModifierLibrary/Source/AnimationModifierLibrary/Private/MotionExtractorModifier.cpp
+++ b/Engine/Plugins/Animation/AnimationModifierLibrary/Source/AnimationModifierLibrary/Private/MotionExtractorModifier.cpp
@@ -65,7 +65,7 @@ void UMotionExtractorModifier::OnApply_Implementation(UAnimSequence* Animation)
 	RequiredBones.Add(BoneIndex);
 	Skeleton->GetReferenceSkeleton().EnsureParentsExistAndSort(RequiredBones);
 
-	FBoneContainer BoneContainer(RequiredBones, false, *Skeleton);
+	FBoneContainer BoneContainer(RequiredBones, UE::Anim::ECurveFilterMode::DisallowAll, *Skeleton);
 	const FCompactPoseBoneIndex CompactPoseBoneIndex = BoneContainer.MakeCompactPoseIndex(FMeshPoseBoneIndex(BoneIndex));
 
 	FTransform FirstFrameBoneTransform = UMotionExtractorUtilityLibrary::ExtractBoneTransform(Animation, BoneContainer, CompactPoseBoneIndex, 0.f, bComponentSpace);
diff --git a/Engine/Plugins/Animation/AnimationWarping/Source/Runtime/Private/BoneControllers/AnimNode_FootPlacement.cpp b/Engine/Plugins/Animation/AnimationWarping/Source/Runtime/Private/BoneControllers/AnimNode_FootPlacement.cpp
--- a/Engine/Plugins/Animation/AnimationWarping/Source/Runtime/Private/BoneControllers/AnimNode_FootPlacement.cpp
+++ b/Engine/Plugins/Animation/AnimationWarping/Source/Runtime/Private/BoneControllers/AnimNode_FootPlacement.cpp
@@ -1262,11 +1262,8 @@ void FAnimNode_FootPlacement::InitializeBoneReferences(const FBoneContainer& Req
 		LegData.InputPose.FootToGround = FTransform(FKFootTransformCS.GetRotation().UnrotateVector(FootAlignmentOffsetCS));
 #endif 
 
-		const USkeleton* Skeleton = RequiredBones.GetSkeletonAsset();
-		check(Skeleton);
-		// Grab UIDs of filtered curves to avoid lookup later
-		LegData.SpeedCurveUID = Skeleton->GetUIDByName(USkeleton::AnimCurveMappingName, LegDef.SpeedCurveName);
-		LegData.DisableLockCurveUID = Skeleton->GetUIDByName(USkeleton::AnimCurveMappingName, LegDef.DisableLockCurveName);
+		LegData.SpeedCurveName = LegDef.SpeedCurveName;
+		LegData.DisableLockCurveName = LegDef.DisableLockCurveName;
 	}
 
 	PelvisBone.Initialize(RequiredBones);
@@ -1352,11 +1349,11 @@ void FAnimNode_FootPlacement::GatherLegDataFromInputs(
 		// If the curve is not found in the stream, assume we're unplanted.
 		const float DefaultSpeedCurveValue = PlantSettings.SpeedThreshold;
 		LegData.InputPose.Speed =
-			Context.CSPContext.Curve.Get(LegData.SpeedCurveUID, bValidSpeedCurve, DefaultSpeedCurveValue);
+			Context.CSPContext.Curve.Get(LegData.SpeedCurveName, bValidSpeedCurve, DefaultSpeedCurveValue);
 	}
 
 	// Grab the lock curve's alpha. If the curve isn't set, then LockAlpha is full weight.
-	LegData.InputPose.LockAlpha = 1.0f - Context.CSPContext.Curve.Get(LegData.DisableLockCurveUID);
+	LegData.InputPose.LockAlpha = 1.0f - Context.CSPContext.Curve.Get(LegData.DisableLockCurveName);
 
 	LegData.InputPose.DistanceToPlant = CalcTargetPlantPlaneDistance(Context, LegData.InputPose);
 	const float FKAlignmentAlpha = GetAlignmentAlpha(Context, LegData.InputPose);
diff --git a/Engine/Plugins/Animation/AnimationWarping/Source/Runtime/Public/BoneControllers/AnimNode_FootPlacement.h b/Engine/Plugins/Animation/AnimationWarping/Source/Runtime/Public/BoneControllers/AnimNode_FootPlacement.h
--- a/Engine/Plugins/Animation/AnimationWarping/Source/Runtime/Public/BoneControllers/AnimNode_FootPlacement.h
+++ b/Engine/Plugins/Animation/AnimationWarping/Source/Runtime/Public/BoneControllers/AnimNode_FootPlacement.h
@@ -38,9 +38,14 @@ namespace UE::Anim::FootPlacement
 		} Bones;
 
 		// Curves
+		UE_DEPRECATED(5.3, "SpeedCurveUID is no longer used, use SpeedCurveName instead.")
 		SmartName::UID_Type SpeedCurveUID = SmartName::MaxUID;
+		UE_DEPRECATED(5.3, "DisableLockCurveUID is no longer used, use DisableLockCurveName instead.")
 		SmartName::UID_Type DisableLockCurveUID = SmartName::MaxUID;
 
+		FName SpeedCurveName = NAME_None;
+		FName DisableLockCurveName = NAME_None;
+
 		// Helper struct to store values coming directly, or trivial to calculate from just the input pose.
 		struct FInputPoseData
 		{
diff --git a/Engine/Plugins/Animation/ControlRig/Source/ControlRig/Private/AnimNode_ControlRig.cpp b/Engine/Plugins/Animation/ControlRig/Source/ControlRig/Private/AnimNode_ControlRig.cpp
--- a/Engine/Plugins/Animation/ControlRig/Source/ControlRig/Private/AnimNode_ControlRig.cpp
+++ b/Engine/Plugins/Animation/ControlRig/Source/ControlRig/Private/AnimNode_ControlRig.cpp
@@ -8,6 +8,7 @@
 #include "GameFramework/Actor.h"
 #include "Animation/NodeMappingContainer.h"
 #include "AnimationRuntime.h"
+#include "Animation/AnimCurveUtils.h"
 
 #include UE_INLINE_GENERATED_CPP_BY_NAME(AnimNode_ControlRig)
 
@@ -130,19 +131,16 @@ void FAnimNode_ControlRig::CacheBones_AnyThread(const FAnimationCacheBonesContex
 	FAnimNode_ControlRigBase::CacheBones_AnyThread(Context);
 
 	FBoneContainer& RequiredBones = Context.AnimInstanceProxy->GetRequiredBones();
-	InputToCurveMappingUIDs.Reset();
 	InputToControlIndex.Reset();
 
 	if(RequiredBones.IsValid())
 	{
 		RefPoseSetterHash.Reset();
 
-		const FSmartNameMapping* CurveMapping = RequiredBones.GetSkeletonAsset()->GetSmartNameContainer(USkeleton::AnimCurveMappingName);
-
-		auto CacheMapping = [&](const TMap<FName, FName>& Mapping, const FSmartNameMapping* CurveNameMapping, 
+		auto CacheMapping = [&](const TMap<FName, FName>& InMapping, FCurveMappings& InCurveMappings,
 			const FAnimationCacheBonesContext& InContext, URigHierarchy* InHierarchy)
 		{
-			for (auto Iter = Mapping.CreateConstIterator(); Iter; ++Iter)
+			for (auto Iter = InMapping.CreateConstIterator(); Iter; ++Iter)
 			{
 				// we need to have list of variables using pin
 				const FName SourcePath = Iter.Key();
@@ -150,14 +148,9 @@ void FAnimNode_ControlRig::CacheBones_AnyThread(const FAnimationCacheBonesContex
 
 				if (SourcePath != NAME_None && TargetPath != NAME_None)
 				{
-					const SmartName::UID_Type Found = CurveNameMapping->FindUID(TargetPath);
-					if (Found != SmartName::MaxUID)
-					{
-						// set value - sound should be UID
-						InputToCurveMappingUIDs.Add(Iter.Value()) = Found;
-						continue;
-					}
-					else if(InHierarchy)
+					InCurveMappings.Add(SourcePath, TargetPath);
+
+					if(InHierarchy)
 					{
 						const FRigElementKey Key(TargetPath, ERigElementType::Control);
 						if(const FRigControlElement* ControlElement = InHierarchy->Find<FRigControlElement>(Key))
@@ -166,9 +159,6 @@ void FAnimNode_ControlRig::CacheBones_AnyThread(const FAnimationCacheBonesContex
 							continue;
 						}
 					}
-
-					UE_LOG(LogAnimation, Warning, TEXT("Curve %s Not Found from the Skeleton %s"), 
-						*TargetPath.ToString(), *GetNameSafe(InContext.AnimInstanceProxy->GetSkeleton()));
 				}
 
 				// @todo: should we clear the item if not found?
@@ -181,8 +171,8 @@ void FAnimNode_ControlRig::CacheBones_AnyThread(const FAnimationCacheBonesContex
 			Hierarchy = CurrentControlRig->GetHierarchy();
 		}
 
-		CacheMapping(InputMapping, CurveMapping, Context, Hierarchy);
-		CacheMapping(OutputMapping, CurveMapping, Context, Hierarchy);
+		CacheMapping(InputMapping, InputCurveMappings, Context, Hierarchy);
+		CacheMapping(OutputMapping, OutputCurveMappings, Context, Hierarchy);
 	}
 }
 
@@ -219,31 +209,19 @@ void FAnimNode_ControlRig::UpdateInput(UControlRig* InControlRig, const FPoseCon
 	// now go through variable mapping table and see if anything is mapping through input
 	if (InputMapping.Num() > 0 && InControlRig)
 	{
-		for (auto Iter = InputMapping.CreateConstIterator(); Iter; ++Iter)
-		{
-			// we need to have list of variables using pin
-			const FName SourcePath = Iter.Key();
-			if (SourcePath != NAME_None)
+		UE::Anim::FCurveUtils::BulkGet(InOutput.Curve, InputCurveMappings, 
+			[&InControlRig](const FControlRigCurveMapping& InBulkElement, float InValue)
 			{
-				const FName CurveName = Iter.Value();
-
-				SmartName::UID_Type UID = *InputToCurveMappingUIDs.Find(CurveName);
-				if (UID != SmartName::MaxUID)
+				FRigVMExternalVariable Variable = InControlRig->GetPublicVariableByName(InBulkElement.SourceName);
+				if (!Variable.bIsReadOnly && Variable.TypeName == TEXT("float"))
 				{
-					const float Value = InOutput.Curve.Get(UID);
-
-					FRigVMExternalVariable Variable = InControlRig->GetPublicVariableByName(SourcePath);
-					if (!Variable.bIsReadOnly && Variable.TypeName == TEXT("float"))
-					{
-						Variable.SetValue<float>(Value);
-					}
-					else
-					{
-						UE_LOG(LogAnimation, Warning, TEXT("[%s] Missing Input Variable [%s]"), *GetNameSafe(InControlRig->GetClass()), *SourcePath.ToString());
-					}
+					Variable.SetValue<float>(InValue);
 				}
-			}
-		} 
+				else
+				{
+					UE_LOG(LogAnimation, Warning, TEXT("[%s] Missing Input Variable [%s]"), *GetNameSafe(InControlRig->GetClass()), *InBulkElement.SourceName.ToString());
+				}
+			});
 	}
 }
 
@@ -255,30 +233,21 @@ void FAnimNode_ControlRig::UpdateOutput(UControlRig* InControlRig, FPoseContext&
 
 	if (OutputMapping.Num() > 0 && InControlRig)
 	{
-		for (auto Iter = OutputMapping.CreateConstIterator(); Iter; ++Iter)
-		{
-			// we need to have list of variables using pin
-			const FName SourcePath = Iter.Key();
-			const FName CurveName = Iter.Value();
-
-			if (SourcePath != NAME_None)
+		UE::Anim::FCurveUtils::BulkSet(InOutput.Curve, OutputCurveMappings, 
+			[&InControlRig](const FControlRigCurveMapping& InBulkElement) -> float
 			{
-				FRigVMExternalVariable Variable = InControlRig->GetPublicVariableByName(SourcePath);
+				FRigVMExternalVariable Variable = InControlRig->GetPublicVariableByName(InBulkElement.SourceName);
 				if (Variable.TypeName == TEXT("float"))
 				{
-					float Value = Variable.GetValue<float>();
-					SmartName::UID_Type* UID = InputToCurveMappingUIDs.Find(Iter.Value());
-					if (UID)
-					{
-						InOutput.Curve.Set(*UID, Value);
-					}
+					return Variable.GetValue<float>();
 				}
 				else
 				{
-					UE_LOG(LogAnimation, Warning, TEXT("[%s] Missing Output Variable [%s]"), *GetNameSafe(ControlRig->GetClass()), *SourcePath.ToString());
+					UE_LOG(LogAnimation, Warning, TEXT("[%s] Missing Output Variable [%s]"), *GetNameSafe(InControlRig->GetClass()), *InBulkElement.SourceName.ToString());
 				}
-			}
-		}
+				
+				return 0.0f;
+			});
 	}
 }
 
diff --git a/Engine/Plugins/Animation/ControlRig/Source/ControlRig/Private/AnimNode_ControlRigBase.cpp b/Engine/Plugins/Animation/ControlRig/Source/ControlRig/Private/AnimNode_ControlRigBase.cpp
--- a/Engine/Plugins/Animation/ControlRig/Source/ControlRig/Private/AnimNode_ControlRigBase.cpp
+++ b/Engine/Plugins/Animation/ControlRig/Source/ControlRig/Private/AnimNode_ControlRigBase.cpp
@@ -11,6 +11,7 @@
 #include "Units/Execution/RigUnit_BeginExecution.h"
 #include "Units/Execution/RigUnit_PrepareForExecution.h"
 #include "Algo/Transform.h"
+#include "Animation/AnimCurveUtils.h"
 
 #include UE_INLINE_GENERATED_CPP_BY_NAME(AnimNode_ControlRigBase)
 
@@ -76,6 +77,13 @@ void FAnimNode_ControlRigBase::Initialize_AnyThread(const FAnimationInitializeCo
 		ControlRig->RequestInit();
 		bControlRigRequiresInitialization = true;
 		LastBonesSerialNumberForCacheBones = 0;
+
+		// Cache curves
+		CachedBulkCurves.Empty();
+		for(FRigCurveElement* CurveElement : ControlRig->GetHierarchy()->GetCurves())
+		{
+			CachedBulkCurves.AddIndexed(CurveElement->GetName());
+		}
 	}
 	else
 	{
@@ -238,45 +246,13 @@ void FAnimNode_ControlRigBase::UpdateInput(UControlRig* ControlRig, const FPoseC
 
 	if (InputSettings.bUpdateCurves && bTransferInputCurves)
 	{
-		if(!ControlRigCurveMappingByIndex.IsEmpty())
-		{
-			for (const TPair<uint16, uint16>& Pair : ControlRigCurveMappingByIndex)
-			{
-				const uint16 ControlRigIndex = Pair.Key;
-				const uint16 SkeletonIndex = Pair.Value;
-
-				bool bIsValid;
-				const float Value = InOutput.Curve.Get(SkeletonIndex, bIsValid);
-				if (bIsValid)
-				{
-					ControlRig->GetHierarchy()->SetCurveValueByIndex(ControlRigIndex, Value);
-				}
-				else
-				{
-					ControlRig->GetHierarchy()->UnsetCurveValueByIndex(ControlRigIndex);
-				}
-			}
-		}
-		else
+		// TODO: Do we need to 'unset' old values manually here? The previous code could do this because it knows the
+		// global set of curves, but we dont know that any more
+		UE::Anim::FCurveUtils::BulkGet(InOutput.Curve, CachedBulkCurves, [ControlRig](const UE::Anim::FNamedIndexElement& InBulkElement, float InValue)
 		{
-			for (auto Iter = ControlRigCurveMappingByName.CreateConstIterator(); Iter; ++Iter)
-			{
-				const FName& Name = Iter.Key();
-				const uint16 SkeletonIndex = Iter.Value();
-				const FRigElementKey Key(Name, ERigElementType::Curve);
-
-				bool bIsValid;
-				const float Value = InOutput.Curve.Get(SkeletonIndex, bIsValid);
-				if (bIsValid)
-				{
-					ControlRig->GetHierarchy()->SetCurveValue(Key, Value);
-				}
-				else
-				{
-					ControlRig->GetHierarchy()->UnsetCurveValue(Key);
-				}
-			}
-		}
+			const FRigElementKey Key(InBulkElement.Name, ERigElementType::Curve);
+			ControlRig->GetHierarchy()->SetCurveValue(Key, InValue);
+		});
 	}
 
 #if WITH_EDITOR
@@ -388,49 +364,11 @@ void FAnimNode_ControlRigBase::UpdateOutput(UControlRig* ControlRig, FPoseContex
 
 	if (OutputSettings.bUpdateCurves)
 	{
-		if(!ControlRigCurveMappingByIndex.IsEmpty())
-		{
-			for (const TPair<uint16, uint16>& Pair : ControlRigCurveMappingByIndex)
-			{
-				const uint16 ControlRigIndex = Pair.Key;
-				const uint16 SkeletonIndex = Pair.Value;
-
-				if (ControlRig->GetHierarchy()->IsCurveValueSetByIndex(ControlRigIndex))
-				{
-					InOutput.Curve.Set(SkeletonIndex, ControlRig->GetHierarchy()->GetCurveValueByIndex(ControlRigIndex));
-				}
-				else
-				{
-					const int32 WeightIndex = InOutput.Curve.GetArrayIndexByUID(SkeletonIndex);
-					if (WeightIndex != INDEX_NONE)
-					{
-						InOutput.Curve.ValidCurveWeights[WeightIndex] = false;
-					}
-				}
-			}
-		}
-		else
+		UE::Anim::FCurveUtils::BulkSet(InOutput.Curve, CachedBulkCurves, [ControlRig](const UE::Anim::FNamedIndexElement& InBulkElement)
 		{
-			for (auto Iter = ControlRigCurveMappingByName.CreateConstIterator(); Iter; ++Iter)
-			{
-				const FName& Name = Iter.Key();
-				const uint16 Index = Iter.Value();
-				const FRigElementKey Key(Name, ERigElementType::Curve);
-
-				if (ControlRig->GetHierarchy()->IsCurveValueSet(Key))
-				{
-					InOutput.Curve.Set(Index, ControlRig->GetHierarchy()->GetCurveValue(Key));
-				}
-				else
-				{
-					const int32 WeightIndex = InOutput.Curve.GetArrayIndexByUID(Index);
-					if (WeightIndex != INDEX_NONE)
-					{
-						InOutput.Curve.ValidCurveWeights[WeightIndex] = false;
-					}
-				}
-			}
-		}
+			FRigCurveElement* CurveElement = ControlRig->GetHierarchy()->GetCurves()[InBulkElement.Index];
+			return ControlRig->GetHierarchy()->GetCurveValue(CurveElement);
+		});
 	}
 
 #if WITH_EDITOR
@@ -717,28 +655,6 @@ void FAnimNode_ControlRigBase::CacheBones_AnyThread(const FAnimationCacheBonesCo
 					}
 				}
 			}
-			
-			// we just support curves by name only
-			{
-				const USkeleton* Skeleton = RequiredBones.GetSkeletonAsset();
-				if (IsValid(Skeleton))
-				{
-					const FSmartNameMapping* CurveMapping = Skeleton->GetSmartNameContainer(USkeleton::AnimCurveMappingName);
-					URigHierarchy* Hierarchy = ControlRig->GetHierarchy();
-					CurveMapping->Iterate([Hierarchy, this](const FSmartNameMappingIterator& Iterator)
-					{
-						// see if the curve name exists in the control rig
-						FName CurveName;
-						if (Iterator.GetName(CurveName))
-						{
-							if (Hierarchy->GetIndex(FRigElementKey(CurveName, ERigElementType::Curve)) != INDEX_NONE)
-							{
-								ControlRigCurveMappingByName.Add(CurveName, Iterator.GetIndex());
-							}
-						}
-					});
-				}
-			}
 
 			// check if we can switch the bones to an index based mapping.
 			// we can only do that if there is no node mapping container set.
@@ -779,32 +695,6 @@ void FAnimNode_ControlRigBase::CacheBones_AnyThread(const FAnimationCacheBonesCo
 					}
 				}
 			}
-
-			bool bIsCurveMappingByIndex = true;
-			
-			// check if we can switch the curves to a index based mapping as well
-			for (auto Iter = ControlRigCurveMappingByName.CreateConstIterator(); Iter; ++Iter)
-			{
-				const uint16 SkeletonIndex = Iter.Value();
-				const int32 ControlRigIndex = ControlRig->GetHierarchy()->GetIndex(FRigElementKey(Iter.Key(), ERigElementType::Curve));
-				if(ControlRigIndex != INDEX_NONE)
-				{
-					ControlRigCurveMappingByIndex.Add(TPair<uint16, uint16>((uint16)ControlRigIndex, SkeletonIndex));
-				}
-				else
-				{
-					bIsCurveMappingByIndex = false;
-				}
-			}
-
-			if(bIsCurveMappingByIndex)
-			{
-				ControlRigCurveMappingByName.Reset();
-			}
-			else
-			{
-				ControlRigCurveMappingByIndex.Reset();
-			}
 		}
 
 		if(bControlRigRequiresInitialization)
diff --git a/Engine/Plugins/Animation/ControlRig/Source/ControlRig/Private/ControlRigAnimInstance.cpp b/Engine/Plugins/Animation/ControlRig/Source/ControlRig/Private/ControlRigAnimInstance.cpp
--- a/Engine/Plugins/Animation/ControlRig/Source/ControlRig/Private/ControlRigAnimInstance.cpp
+++ b/Engine/Plugins/Animation/ControlRig/Source/ControlRig/Private/ControlRigAnimInstance.cpp
@@ -45,7 +45,7 @@ bool FControlRigAnimInstanceProxy::Evaluate(FPoseContext& Output)
 		Output.Pose[ModifiedBone] = CompactPose[ModifiedBone];
 	}
 
-	for (TPair<SmartName::UID_Type, float> Pair : StoredCurves)
+	for (const TPair<FName, float>& Pair : StoredCurves)
 	{
 		Output.Curve.Set(Pair.Key, Pair.Value);
 	}
diff --git a/Engine/Plugins/Animation/ControlRig/Source/ControlRig/Private/ControlRigComponent.cpp b/Engine/Plugins/Animation/ControlRig/Source/ControlRig/Private/ControlRigComponent.cpp
--- a/Engine/Plugins/Animation/ControlRig/Source/ControlRig/Private/ControlRigComponent.cpp
+++ b/Engine/Plugins/Animation/ControlRig/Source/ControlRig/Private/ControlRigComponent.cpp
@@ -676,18 +676,10 @@ void UControlRigComponent::AddMappedSkeletalMesh(USkeletalMeshComponent* Skeleta
 			{
 				CR->GetHierarchy()->ForEach<FRigCurveElement>([Skeleton, &CurvesToMap](FRigCurveElement* CurveElement) -> bool
                 {
-                    const FSmartNameMapping* CurveNameMapping = Skeleton->GetSmartNameContainer(USkeleton::AnimCurveMappingName);
-					if (CurveNameMapping)
-					{
-						FSmartName SmartName;
-						if (CurveNameMapping->FindSmartName(CurveElement->GetName(), SmartName))
-						{
-							FControlRigComponentMappedCurve CurveToMap;
-							CurveToMap.Source = CurveElement->GetName();
-							CurveToMap.Target = CurveElement->GetName();
-							CurvesToMap.Add(CurveToMap);
-						}
-					}
+					FControlRigComponentMappedCurve CurveToMap;
+					CurveToMap.Source = CurveElement->GetName();
+					CurveToMap.Target = CurveElement->GetName();
+					CurvesToMap.Add(CurveToMap);
 					return true;
 				});
 			}
@@ -1357,19 +1349,7 @@ void UControlRigComponent::ValidateMappingData()
 					{
 						if (USkeleton* Skeleton = SkeletalMesh->GetSkeleton())
 						{
-							if (MappedElement.ElementType == ERigElementType::Curve)
-							{
-								const FSmartNameMapping* CurveNameMapping = Skeleton->GetSmartNameContainer(USkeleton::AnimCurveMappingName);
-								if (CurveNameMapping)
-								{
-									FSmartName SmartName;
-									if (CurveNameMapping->FindSmartName(MappedElement.TransformName, SmartName))
-									{
-										MappedElement.SubIndex = (int32)SmartName.UID;
-									}
-								}
-							}
-							else
+							if (MappedElement.ElementType != ERigElementType::Curve)
 							{
 								MappedElement.SubIndex = Skeleton->GetReferenceSkeleton().FindBoneIndex(MappedElement.TransformName);
 							}
@@ -1618,7 +1598,7 @@ void UControlRigComponent::TransferOutputs()
 
 					if (Proxy)
 					{
-						Proxy->StoredCurves.FindOrAdd((SmartName::UID_Type)MappedElement.SubIndex) = ControlRig->GetHierarchy()->GetCurveValue(MappedElement.ElementIndex);
+						Proxy->StoredCurves.FindOrAdd(MappedElement.TransformName) = ControlRig->GetHierarchy()->GetCurveValue(MappedElement.ElementIndex);
 					}
 				}
 			}
diff --git a/Engine/Plugins/Animation/ControlRig/Source/ControlRig/Private/Rigs/AdditiveControlRig.cpp b/Engine/Plugins/Animation/ControlRig/Source/ControlRig/Private/Rigs/AdditiveControlRig.cpp
--- a/Engine/Plugins/Animation/ControlRig/Source/ControlRig/Private/Rigs/AdditiveControlRig.cpp
+++ b/Engine/Plugins/Animation/ControlRig/Source/ControlRig/Private/Rigs/AdditiveControlRig.cpp
@@ -99,7 +99,7 @@ void UAdditiveControlRig::Initialize(bool bInitRigUnits /*= true*/)
 	Execute_Internal(FRigUnit_PrepareForExecution::EventName);
 }
 
-void UAdditiveControlRig::CreateRigElements(const FReferenceSkeleton& InReferenceSkeleton, const FSmartNameMapping* InSmartNameMapping)
+void UAdditiveControlRig::CreateRigElements(const FReferenceSkeleton& InReferenceSkeleton, const USkeleton* InSkeleton)
 {
 	PostInitInstanceIfRequired();
 	
@@ -108,14 +108,12 @@ void UAdditiveControlRig::CreateRigElements(const FReferenceSkeleton& InReferenc
 	{
 		Controller->ImportBones(InReferenceSkeleton, NAME_None, false, false, true, false);
 
-		if (InSmartNameMapping)
+		if (InSkeleton)
 		{
-			TArray<FName> NameArray;
-			InSmartNameMapping->FillNameArray(NameArray);
-			for (int32 Index = 0; Index < NameArray.Num(); ++Index)
+			InSkeleton->ForEachCurveMetaData([Controller](const FName& InCurveName, const FCurveMetaData& InMetaData)
 			{
-				 Controller->AddCurve(NameArray[Index], 0.f, false);
-			}
+				Controller->AddCurve(InCurveName, 0.f, false);
+			});
 		}
 
 		// add control for all bone hierarchy
@@ -156,7 +154,7 @@ void UAdditiveControlRig::CreateRigElements(const USkeletalMesh* InReferenceMesh
 	if (InReferenceMesh)
 	{
 		const USkeleton* Skeleton = InReferenceMesh->GetSkeleton();
-		CreateRigElements(InReferenceMesh->GetRefSkeleton(), (Skeleton) ? Skeleton->GetSmartNameContainer(USkeleton::AnimCurveMappingName) : nullptr);
+		CreateRigElements(InReferenceMesh->GetRefSkeleton(), Skeleton);
 	}
 }
 
diff --git a/Engine/Plugins/Animation/ControlRig/Source/ControlRig/Private/Rigs/FKControlRig.cpp b/Engine/Plugins/Animation/ControlRig/Source/ControlRig/Private/Rigs/FKControlRig.cpp
--- a/Engine/Plugins/Animation/ControlRig/Source/ControlRig/Private/Rigs/FKControlRig.cpp
+++ b/Engine/Plugins/Animation/ControlRig/Source/ControlRig/Private/Rigs/FKControlRig.cpp
@@ -298,7 +298,7 @@ void UFKControlRig::Initialize(bool bInitRigUnits /*= true*/)
 	else if (USkeleton* Skeleton = Cast<USkeleton>(GetObjectBinding()->GetBoundObject()))
 	{
 		const FReferenceSkeleton& RefSkeleton = Skeleton->GetReferenceSkeleton();
-		CreateRigElements(RefSkeleton, (Skeleton) ? Skeleton->GetSmartNameContainer(USkeleton::AnimCurveMappingName) : nullptr);
+		CreateRigElements(RefSkeleton, Skeleton);
 	}
 }
 
@@ -340,7 +340,7 @@ void UFKControlRig::SetControlActive(const TArray<FFKBoneCheckInfo>& BoneChecks)
 	}
 }
 
-void UFKControlRig::CreateRigElements(const FReferenceSkeleton& InReferenceSkeleton, const FSmartNameMapping* InSmartNameMapping)
+void UFKControlRig::CreateRigElements(const FReferenceSkeleton& InReferenceSkeleton, const USkeleton* InSkeleton)
 {
 	PostInitInstanceIfRequired();
 
@@ -356,23 +356,18 @@ void UFKControlRig::CreateRigElements(const FReferenceSkeleton& InReferenceSkele
             {
             	for (int32 Index = 0; Index < InitializationOptions.CurveNames.Num(); ++Index)
             	{
-            		if(InSmartNameMapping && InSmartNameMapping->FindUID(InitializationOptions.CurveNames[Index]) != SmartName::MaxUID)
-            		{
-            			Controller->AddCurve(InitializationOptions.CurveNames[Index], 0.f, false);
-            		}
+            		Controller->AddCurve(InitializationOptions.CurveNames[Index], 0.f, false);
             	}
             }
 			// add all curves found on the skeleton
-            else if (InSmartNameMapping && InitializationOptions.bImportCurves)
+            else if (InSkeleton && InitializationOptions.bImportCurves)
             {
-            	TArray<FName> NameArray;
-            	InSmartNameMapping->FillNameArray(NameArray);
-            	for (int32 Index = 0; Index < NameArray.Num(); ++Index)
-            	{
-            		Controller->AddCurve(NameArray[Index], 0.f, false);
-            	}
+            	InSkeleton->ForEachCurveMetaData([Controller](const FName& InCurveName, const FCurveMetaData& InMetaData)
+				{
+					Controller->AddCurve(InCurveName, 0.f, false);
+				});
             }
-		}		
+		}
 
 		// add control for all bone hierarchy 
 		if (InitializationOptions.bGenerateBoneControls)
@@ -498,7 +493,7 @@ void UFKControlRig::CreateRigElements(const USkeletalMesh* InReferenceMesh)
 	if (InReferenceMesh)
 	{
 		const USkeleton* Skeleton = InReferenceMesh->GetSkeleton();
-		CreateRigElements(InReferenceMesh->GetRefSkeleton(), (Skeleton) ? Skeleton->GetSmartNameContainer(USkeleton::AnimCurveMappingName) : nullptr);
+		CreateRigElements(InReferenceMesh->GetRefSkeleton(), Skeleton);
 	}
 }
 
diff --git a/Engine/Plugins/Animation/ControlRig/Source/ControlRig/Private/Rigs/RigHierarchyController.cpp b/Engine/Plugins/Animation/ControlRig/Source/ControlRig/Private/Rigs/RigHierarchyController.cpp
--- a/Engine/Plugins/Animation/ControlRig/Source/ControlRig/Private/Rigs/RigHierarchyController.cpp
+++ b/Engine/Plugins/Animation/ControlRig/Source/ControlRig/Private/Rigs/RigHierarchyController.cpp
@@ -867,29 +867,24 @@ TArray<FRigElementKey> URigHierarchyController::ImportCurves(USkeleton* InSkelet
 		return Keys;
 	}
 
-	const FSmartNameMapping* SmartNameMapping = InSkeleton->GetSmartNameContainer(USkeleton::AnimCurveMappingName);
-
-	TArray<FName> NameArray;
-	SmartNameMapping->FillNameArray(NameArray);
-
-	for (int32 Index = 0; Index < NameArray.Num(); ++Index)
+	InSkeleton->ForEachCurveMetaData([this, InNameSpace, &Keys, bSetupUndo](const FName& InCurveName, const FCurveMetaData& InMetaData)
 	{
-		FName Name = NameArray[Index];
+		FName Name = InCurveName;
 		if (!InNameSpace.IsNone())
 		{
-			Name = *FString::Printf(TEXT("%s::%s"), *InNameSpace.ToString(), *Name.ToString());
+			Name = *FString::Printf(TEXT("%s::%s"), *InNameSpace.ToString(), *InCurveName.ToString());
 		}
 
 		const FRigElementKey ExpectedKey(Name, ERigElementType::Curve);
 		if(Hierarchy->Contains(ExpectedKey))
 		{
 			Keys.Add(ExpectedKey);
-			continue;
+			return;
 		}
 		
 		const FRigElementKey CurveKey = AddCurve(Name, 0.f, bSetupUndo);
 		Keys.Add(FRigElementKey(Name, ERigElementType::Curve));
-	}
+	});
 
 	if(bSelectCurves)
 	{
diff --git a/Engine/Plugins/Animation/ControlRig/Source/ControlRig/Private/Sequencer/MovieSceneControlRigParameterSection.cpp b/Engine/Plugins/Animation/ControlRig/Source/ControlRig/Private/Sequencer/MovieSceneControlRigParameterSection.cpp
--- a/Engine/Plugins/Animation/ControlRig/Source/ControlRig/Private/Sequencer/MovieSceneControlRigParameterSection.cpp
+++ b/Engine/Plugins/Animation/ControlRig/Source/ControlRig/Private/Sequencer/MovieSceneControlRigParameterSection.cpp
@@ -2981,7 +2981,7 @@ bool UMovieSceneControlRigParameterSection::LoadAnimSequenceIntoThisSection(UAni
 		for (const FFloatCurve& Curve : CurveData.FloatCurves)
 		{
 			const float Val = Curve.FloatCurve.Eval(SequenceSecond);
-			SourceHierarchy->SetCurveValue(FRigElementKey(Curve.Name.DisplayName, ERigElementType::Curve), Val);
+			SourceHierarchy->SetCurveValue(FRigElementKey(Curve.GetName(), ERigElementType::Curve), Val);
 		}
 
 		// retrieve the pose using the services that persona and sequencer rely on
diff --git a/Engine/Plugins/Animation/ControlRig/Source/ControlRig/Public/AnimNode_ControlRigBase.h b/Engine/Plugins/Animation/ControlRig/Source/ControlRig/Public/AnimNode_ControlRigBase.h
--- a/Engine/Plugins/Animation/ControlRig/Source/ControlRig/Public/AnimNode_ControlRigBase.h
+++ b/Engine/Plugins/Animation/ControlRig/Source/ControlRig/Public/AnimNode_ControlRigBase.h
@@ -2,6 +2,7 @@
 
 #pragma once
 
+#include "Animation/AnimBulkCurves.h"
 #include "Animation/AnimNode_CustomProperty.h"
 #include "AnimNode_ControlRigBase.generated.h"
 
@@ -122,19 +123,21 @@ protected:
 	TArray<TPair<uint16, uint16>> ControlRigBoneInputMappingByIndex;
 	TArray<TPair<uint16, uint16>> ControlRigBoneOutputMappingByIndex;
 
-	/** Complete mapping from skeleton to curve index */
-	TArray<TPair<uint16, uint16>> ControlRigCurveMappingByIndex;
+	/** Complete mapping from skeleton to curve name */
+	TArray<TPair<uint16, FName>> ControlRigCurveMappingByIndex;
 
 	/** Rig Hierarchy bone name to required array index mapping */
 	TMap<FName, uint16> ControlRigBoneInputMappingByName;
 	TMap<FName, uint16> ControlRigBoneOutputMappingByName;
 
-	/** Rig Curve name to Curve LUI mapping */
-	TMap<FName, uint16> ControlRigCurveMappingByName;
-
-	TMap<FName, uint16> InputToCurveMappingUIDs;
+	/** Rig Curve name to Curve mapping */
+	TMap<FName, FName> ControlRigCurveMappingByName;
+	
 	TMap<FName, int32> InputToControlIndex;
 
+	// Cached bulk curves for I/O
+	UE::Anim::FBulkCurves CachedBulkCurves;
+	
 	/** Node Mapping Container */
 	UPROPERTY(transient)
 	TWeakObjectPtr<UNodeMappingContainer> NodeMappingContainer;
diff --git a/Engine/Plugins/Animation/ControlRig/Source/ControlRig/Public/ControlRigAnimInstance.h b/Engine/Plugins/Animation/ControlRig/Source/ControlRig/Public/ControlRigAnimInstance.h
--- a/Engine/Plugins/Animation/ControlRig/Source/ControlRig/Public/ControlRigAnimInstance.h
+++ b/Engine/Plugins/Animation/ControlRig/Source/ControlRig/Public/ControlRigAnimInstance.h
@@ -30,7 +30,7 @@ public:
 	virtual void UpdateAnimationNode(const FAnimationUpdateContext& InContext) override;
 
 	TMap<int32, FTransform> StoredTransforms;
-	TMap<SmartName::UID_Type, float> StoredCurves;
+	TMap<FName, float> StoredCurves;
 };
 
 UCLASS(transient, NotBlueprintable)
diff --git a/Engine/Plugins/Animation/ControlRig/Source/ControlRig/Public/Rigs/AdditiveControlRig.h b/Engine/Plugins/Animation/ControlRig/Source/ControlRig/Public/Rigs/AdditiveControlRig.h
--- a/Engine/Plugins/Animation/ControlRig/Source/ControlRig/Public/Rigs/AdditiveControlRig.h
+++ b/Engine/Plugins/Animation/ControlRig/Source/ControlRig/Public/Rigs/AdditiveControlRig.h
@@ -31,5 +31,5 @@ private:
 
 	/** Create RigElements - bone hierarchy and curves - from incoming skeleton */
 	void CreateRigElements(const USkeletalMesh* InReferenceMesh);
-	void CreateRigElements(const FReferenceSkeleton& InReferenceSkeleton, const FSmartNameMapping* InSmartNameMapping);
+	void CreateRigElements(const FReferenceSkeleton& InReferenceSkeleton, const USkeleton* InSkeleton);
 };
\ No newline at end of file
diff --git a/Engine/Plugins/Animation/ControlRig/Source/ControlRig/Public/Rigs/FKControlRig.h b/Engine/Plugins/Animation/ControlRig/Source/ControlRig/Public/Rigs/FKControlRig.h
--- a/Engine/Plugins/Animation/ControlRig/Source/ControlRig/Public/Rigs/FKControlRig.h
+++ b/Engine/Plugins/Animation/ControlRig/Source/ControlRig/Public/Rigs/FKControlRig.h
@@ -88,7 +88,7 @@ private:
 
 	/** Create RigElements - bone hierarchy and curves - from incoming skeleton */
 	void CreateRigElements(const USkeletalMesh* InReferenceMesh);
-	void CreateRigElements(const FReferenceSkeleton& InReferenceSkeleton, const FSmartNameMapping* InSmartNameMapping);
+	void CreateRigElements(const FReferenceSkeleton& InReferenceSkeleton, const USkeleton* InSkeleton);
 	void SetControlOffsetsFromBoneInitials();
 
 	UPROPERTY()
diff --git a/Engine/Plugins/Animation/ControlRig/Source/ControlRigDeveloper/Private/AnimGraphNode_ControlRig.cpp b/Engine/Plugins/Animation/ControlRig/Source/ControlRigDeveloper/Private/AnimGraphNode_ControlRig.cpp
--- a/Engine/Plugins/Animation/ControlRig/Source/ControlRigDeveloper/Private/AnimGraphNode_ControlRig.cpp
+++ b/Engine/Plugins/Animation/ControlRig/Source/ControlRigDeveloper/Private/AnimGraphNode_ControlRig.cpp
@@ -623,8 +623,7 @@ void UAnimGraphNode_ControlRig::GetAvailableMapping(const FName& PathName, TArra
 	OutArray.Reset();
 	if (TargetSkeleton)
 	{
-		const FSmartNameMapping* CurveMapping = TargetSkeleton->GetSmartNameContainer(USkeleton::AnimCurveMappingName);
-		CurveMapping->FillNameArray(OutArray);
+		TargetSkeleton->GetCurveMetaDataNames(OutArray);
 		
 		// also add all controls
 		if(const UClass* ControlRigClass = GetTargetClass())
diff --git a/Engine/Plugins/Animation/GameplayInsights/Source/GameplayInsights/Private/AnimCurvesTrack.cpp b/Engine/Plugins/Animation/GameplayInsights/Source/GameplayInsights/Private/AnimCurvesTrack.cpp
--- a/Engine/Plugins/Animation/GameplayInsights/Source/GameplayInsights/Private/AnimCurvesTrack.cpp
+++ b/Engine/Plugins/Animation/GameplayInsights/Source/GameplayInsights/Private/AnimCurvesTrack.cpp
@@ -118,50 +118,16 @@ TSharedPtr<SCurveTimelineView::FTimelineCurveData> FAnimCurveTrack::GetCurveData
 bool FAnimCurveTrack::UpdateInternal()
 {
 	IRewindDebugger* RewindDebugger = IRewindDebugger::Instance();
-	// compute curve points
-	//
-	// convert time range to from rewind debugger times to profiler times
-	TRange<double> TraceTimeRange = RewindDebugger->GetCurrentTraceRange();
-	double StartTime = TraceTimeRange.GetLowerBoundValue();
-	double EndTime = TraceTimeRange.GetUpperBoundValue();
-	
 	const TraceServices::IAnalysisSession* AnalysisSession = RewindDebugger->GetAnalysisSession();
-	const FGameplayProvider* GameplayProvider = AnalysisSession->ReadProvider<FGameplayProvider>(FGameplayProvider::ProviderName);
 	const FAnimationProvider* AnimationProvider = AnalysisSession->ReadProvider<FAnimationProvider>(FAnimationProvider::ProviderName);
 
-	
-	if(CurvesUpdateRequested > 10 && GameplayProvider && AnimationProvider)
+	if(CurvesUpdateRequested > 10 && AnimationProvider)
 	{
 		CurvesUpdateRequested = 0;
-		auto& CurvePoints = CurveData->Points;
-		CurvePoints.SetNum(0,false);
-	
-		TraceServices::FAnalysisSessionReadScope SessionReadScope(*AnalysisSession);
-
-		AnimationProvider->ReadSkeletalMeshPoseTimeline(ObjectId, [this, AnimationProvider, StartTime, EndTime, AnalysisSession, &CurvePoints](const FAnimationProvider::SkeletalMeshPoseTimeline& InTimeline, bool bHasCurves)
-		{
-			// this isn't very efficient, and it gets called every frame.  will need optimizing
-			InTimeline.EnumerateEvents(StartTime, EndTime, [this, StartTime, EndTime, AnimationProvider, AnalysisSession, &CurvePoints](double InStartTime, double InEndTime, uint32 InDepth, const FSkeletalMeshPoseMessage& InMessage)
-			{
-				if (InEndTime > StartTime && InStartTime < EndTime)
-				{
-					double Time = InMessage.RecordingTime;
-					
-					AnimationProvider->EnumerateSkeletalMeshCurves(InMessage, [this, Time, &CurvePoints](const FSkeletalMeshNamedCurve& InCurve)
-					{
-						if (InCurve.Id == CurveId)
-						{
-							CurvePoints.Add({Time,InCurve.Value});
-						}
-					});
-				}
-				return TraceServices::EEventEnumerate::Continue;
-			});
-		});
+		UpdateCurvePointsInternal();
 	}
-		
-	bool bChanged = false;
 
+	bool bChanged = false;
 	if (CurveName.IsEmpty() && AnimationProvider)
 	{
 		CurveName = FText::FromString(AnimationProvider->GetName(CurveId));
@@ -171,6 +137,44 @@ bool FAnimCurveTrack::UpdateInternal()
 	return bChanged;
 }
 
+void FAnimCurveTrack::UpdateCurvePointsInternal()
+{
+	IRewindDebugger* RewindDebugger = IRewindDebugger::Instance();
+	const TraceServices::IAnalysisSession* AnalysisSession = RewindDebugger->GetAnalysisSession();
+	const FAnimationProvider* AnimationProvider = AnalysisSession->ReadProvider<FAnimationProvider>(FAnimationProvider::ProviderName);
+
+	// convert time range to from rewind debugger times to profiler times
+	TRange<double> TraceTimeRange = RewindDebugger->GetCurrentTraceRange();
+	double StartTime = TraceTimeRange.GetLowerBoundValue();
+	double EndTime = TraceTimeRange.GetUpperBoundValue();
+	
+	auto& CurvePoints = CurveData->Points;
+	CurvePoints.SetNum(0,false);
+
+	TraceServices::FAnalysisSessionReadScope SessionReadScope(*AnalysisSession);
+	
+	AnimationProvider->ReadSkeletalMeshPoseTimeline(ObjectId, [this, AnimationProvider, StartTime, EndTime, AnalysisSession, &CurvePoints](const FAnimationProvider::SkeletalMeshPoseTimeline& InTimeline, bool bHasCurves)
+	{
+		// this isn't very efficient, and it gets called every frame.  will need optimizing
+		InTimeline.EnumerateEvents(StartTime, EndTime, [this, StartTime, EndTime, AnimationProvider, AnalysisSession, &CurvePoints](double InStartTime, double InEndTime, uint32 InDepth, const FSkeletalMeshPoseMessage& InMessage)
+		{
+			if (InEndTime > StartTime && InStartTime < EndTime)
+			{
+				double Time = InMessage.RecordingTime;
+			
+				AnimationProvider->EnumerateSkeletalMeshCurves(InMessage, [this, Time, &CurvePoints](const FSkeletalMeshNamedCurve& InCurve)
+				{
+					if (InCurve.Id == CurveId)
+					{
+						CurvePoints.Add({Time,InCurve.Value});
+					}
+				});
+			}
+			return TraceServices::EEventEnumerate::Continue;
+		});
+	});
+}
+
 TSharedPtr<SWidget> FAnimCurveTrack::GetDetailsViewInternal() 
 {
 	IRewindDebugger* RewindDebugger = IRewindDebugger::Instance();
diff --git a/Engine/Plugins/Animation/GameplayInsights/Source/GameplayInsights/Private/AnimCurvesTrack.h b/Engine/Plugins/Animation/GameplayInsights/Source/GameplayInsights/Private/AnimCurvesTrack.h
--- a/Engine/Plugins/Animation/GameplayInsights/Source/GameplayInsights/Private/AnimCurvesTrack.h
+++ b/Engine/Plugins/Animation/GameplayInsights/Source/GameplayInsights/Private/AnimCurvesTrack.h
@@ -15,8 +15,7 @@ class FAnimCurveTrack : public FRewindDebuggerTrack
 {
 public:
 	FAnimCurveTrack(uint64 InObjectId, uint32 InCurveId);
-
-
+	
 	uint32 GetCurveId() { return CurveId; };
 
 private:
@@ -27,9 +26,10 @@ private:
 	virtual FName GetNameInternal() const override { return "AnimationCurve"; }
 	virtual FText GetDisplayNameInternal() const override { return CurveName; }
 	virtual uint64 GetObjectIdInternal() const override { return ObjectId; }
-
+	virtual void UpdateCurvePointsInternal();
+	
 	TSharedPtr<SCurveTimelineView::FTimelineCurveData> GetCurveData() const;
-
+protected:
 	uint32 CurveId;
 
 	FSlateIcon Icon;
diff --git a/Engine/Plugins/Animation/GameplayInsights/Source/GameplayInsights/Private/AnimationAnalyzer.cpp b/Engine/Plugins/Animation/GameplayInsights/Source/GameplayInsights/Private/AnimationAnalyzer.cpp
--- a/Engine/Plugins/Animation/GameplayInsights/Source/GameplayInsights/Private/AnimationAnalyzer.cpp
+++ b/Engine/Plugins/Animation/GameplayInsights/Source/GameplayInsights/Private/AnimationAnalyzer.cpp
@@ -49,6 +49,7 @@ void FAnimationAnalyzer::OnAnalysisBegin(const FOnAnalysisContext& Context)
 	Builder.RouteEvent(RouteId_Montage2, "Animation", "Montage2");
 	Builder.RouteEvent(RouteId_Sync, "Animation", "Sync");
 	Builder.RouteEvent(RouteId_PoseWatch, "Animation", "PoseWatch");
+	Builder.RouteEvent(RouteId_PoseWatch2, "Animation", "PoseWatch2");
 }
 
 bool FAnimationAnalyzer::OnEvent(uint16 RouteId, EStyle Style, const FOnEventContext& Context)
@@ -427,6 +428,26 @@ bool FAnimationAnalyzer::OnEvent(uint16 RouteId, EStyle Style, const FOnEventCon
 		AnimationProvider.AppendPoseWatch(AnimInstanceId, Context.EventTime.AsSeconds(Cycle), RecordingTime, PoseWatchId, BoneTransformsFloatArray, RequiredBonesIntArray, WorldTransformFloatArray, bIsEnabled);
 		break;
 	}
+	case RouteId_PoseWatch2:
+	{
+		uint64 Cycle = EventData.GetValue<uint64>("Cycle");
+		double RecordingTime = EventData.GetValue<double>("RecordingTime");
+		uint64 AnimInstanceId = EventData.GetValue<uint64>("AnimInstanceId");
+		uint64 ComponentId = EventData.GetValue<uint64>("ComponentId");
+		uint64 PoseWatchId = EventData.GetValue<uint64>("PoseWatchId");
+		uint32 NameId = EventData.GetValue<uint32>("NameId");
+		uint32 Color = EventData.GetValue<uint32>("Color");
+		bool bIsEnabled = EventData.GetValue<bool>("bIsEnabled");
+		TArrayView<const float> WorldTransformFloatArray = EventData.GetArrayView<float>("WorldTransform");
+		TArrayView<const uint16> RequiredBonesIntArray = EventData.GetArrayView<uint16>("RequiredBones");
+		TArrayView<const float> BoneTransformsFloatArray = EventData.GetArrayView<float>("BoneTransforms");
+		TArrayView<const uint32> CurveIds = EventData.GetArrayView<uint32>("CurveIds");
+		TArrayView<const float> CurveValues = EventData.GetArrayView<float>("CurveValues");
+		check(CurveIds.Num() == CurveValues.Num());
+
+		AnimationProvider.AppendPoseWatch(ComponentId, AnimInstanceId, Context.EventTime.AsSeconds(Cycle), RecordingTime, PoseWatchId, NameId, FColor(Color), BoneTransformsFloatArray, CurveIds, CurveValues, RequiredBonesIntArray, WorldTransformFloatArray, bIsEnabled);
+		break;
+	}
 	}
 
 	return true;
diff --git a/Engine/Plugins/Animation/GameplayInsights/Source/GameplayInsights/Private/AnimationAnalyzer.h b/Engine/Plugins/Animation/GameplayInsights/Source/GameplayInsights/Private/AnimationAnalyzer.h
--- a/Engine/Plugins/Animation/GameplayInsights/Source/GameplayInsights/Private/AnimationAnalyzer.h
+++ b/Engine/Plugins/Animation/GameplayInsights/Source/GameplayInsights/Private/AnimationAnalyzer.h
@@ -49,7 +49,8 @@ private:
 		RouteId_Montage,
 		RouteId_Montage2,
 		RouteId_Sync,
-		RouteId_PoseWatch
+		RouteId_PoseWatch,
+		RouteId_PoseWatch2,
 	};
 
 	TraceServices::IAnalysisSession& Session;
