AnimSequence are always rebuilt in Lyra
#fix more deterministic issues (between consecutive runs, I tested single process only previously), as FName indices can vary
#misc use string-base name for Transform curves (missed that when changing float curves previously)

diff --git a/Engine/Plugins/Animation/AnimationData/Source/AnimationData/Private/AnimSequencerDataModel.cpp b/Engine/Plugins/Animation/AnimationData/Source/AnimationData/Private/AnimSequencerDataModel.cpp
--- a/Engine/Plugins/Animation/AnimationData/Source/AnimationData/Private/AnimSequencerDataModel.cpp
+++ b/Engine/Plugins/Animation/AnimationData/Source/AnimationData/Private/AnimSequencerDataModel.cpp
@@ -632,8 +632,11 @@ FGuid UAnimationSequencerDataModel::GenerateGuid() const
 		
 		for (const FAnimatedBoneAttribute& Attribute : AnimatedBoneAttributes)
 		{
-			UpdateWithData(Attribute.Identifier);
-			const uint32 StructSize = Attribute.Identifier.GetType()->GetStructureSize();
+			UpdateSHAWithArray(Attribute.Identifier.GetName().ToString().GetCharArray());
+			UpdateSHAWithArray(Attribute.Identifier.GetBoneName().ToString().GetCharArray());
+			UpdateWithData(Attribute.Identifier.GetBoneIndex());
+			UpdateSHAWithArray(Attribute.Identifier.GetType()->GetFName().ToString().GetCharArray());
+			const uint32 StructSize = Attribute.Identifier.GetType()->GetPropertiesSize();
 			for (const FAttributeKey& Key : Attribute.Curve.GetConstRefOfKeys())
 			{
 				UpdateWithData(Key.Time);
@@ -651,7 +654,8 @@ FGuid UAnimationSequencerDataModel::GenerateGuid() const
 
 		for (const FTransformCurve& Curve : GetTransformCurves())
 		{
-			UpdateWithData(Curve.GetName());
+		   	const FString CurveName = Curve.GetName().ToString();
+    		UpdateSHAWithArray(CurveName.GetCharArray());
 
 			auto UpdateWithComponent = [&UpdateWithFloatCurve](const FVectorCurve& VectorCurve)
 			{
@@ -824,70 +828,70 @@ void UAnimationSequencerDataModel::GenerateLegacyCurveData()
 			{
 				if (const UControlRig* ControlRig = Section->GetControlRig())
 				{
-					if(URigHierarchy* Hierarchy = ControlRig->GetHierarchy())
-					{
-						const FString SequencerSuffix(TEXT("_Sequencer"));
-						const TArray<FScalarParameterNameAndCurve>& ScalarCurves = Section->GetScalarParameterNamesAndCurves();				
-						if (RetainFloatCurves)
+						if(URigHierarchy* Hierarchy = ControlRig->GetHierarchy())
 						{
-							LegacyCurveData.FloatCurves.RemoveAll([SequencerSuffix](const FFloatCurve& FloatCurve)
+							const FString SequencerSuffix(TEXT("_Sequencer"));
+							const TArray<FScalarParameterNameAndCurve>& ScalarCurves = Section->GetScalarParameterNamesAndCurves();				
+							if (RetainFloatCurves)
 							{
+								LegacyCurveData.FloatCurves.RemoveAll([SequencerSuffix](const FFloatCurve& FloatCurve)
+								{
 								return FloatCurve.GetName().ToString().EndsWith(SequencerSuffix);
-							});
-						}
-						else
-						{
-							LegacyCurveData.FloatCurves.Empty();
-						}				
+								});
+							}
+							else
+							{
+								LegacyCurveData.FloatCurves.Empty();
+							}				
 
 						Hierarchy->ForEach<FRigCurveElement>([Hierarchy, this, ScalarCurves, FrameRate = GetFrameRate(), SequencerSuffix](const FRigCurveElement* CurveElement) -> bool
-						{
-							const FRigElementKey ControlKey(UFKControlRig::GetControlName(CurveElement->GetName(), ERigElementType::Curve), ERigElementType::Control);
-							if (const FRigControlElement* Element = Hierarchy->Find<FRigControlElement>(ControlKey))
 							{
-								FFloatCurve& FloatCurve = LegacyCurveData.FloatCurves.AddDefaulted_GetRef();
-								if (RetainFloatCurves)
+								const FRigElementKey ControlKey(UFKControlRig::GetControlName(CurveElement->GetName(), ERigElementType::Curve), ERigElementType::Control);
+								if (const FRigControlElement* Element = Hierarchy->Find<FRigControlElement>(ControlKey))
 								{
+									FFloatCurve& FloatCurve = LegacyCurveData.FloatCurves.AddDefaulted_GetRef();
+									if (RetainFloatCurves)
+									{
 									FloatCurve.SetName(FName(*(CurveElement->GetName().ToString() + TEXT("_Sequencer"))));
-								}
-								else
-								{
+									}
+									else
+									{
 									FloatCurve.SetName(CurveElement->GetName());
-								}						
-
-								FloatCurve.Color = Element->Settings.ShapeColor;
+									}						
 								
+									FloatCurve.Color = Element->Settings.ShapeColor;
+									
 								const FAnimationCurveIdentifier CurveId(FloatCurve.GetName(), ERawCurveTrackTypes::RCT_Float);
 								if (!RetainFloatCurves || !FloatCurve.GetName().ToString().Contains(SequencerSuffix))
-								{
-									if (CurveIdentifierToMetaData.Contains(CurveId))
 									{
-										const FAnimationCurveMetaData& CurveMetaData = CurveIdentifierToMetaData.FindChecked(CurveId);
-										FloatCurve.SetCurveTypeFlags(CurveMetaData.Flags);
-										FloatCurve.Color = CurveMetaData.Color;
+										if (CurveIdentifierToMetaData.Contains(CurveId))
+										{
+											const FAnimationCurveMetaData& CurveMetaData = CurveIdentifierToMetaData.FindChecked(CurveId);
+											FloatCurve.SetCurveTypeFlags(CurveMetaData.Flags);
+											FloatCurve.Color = CurveMetaData.Color;
+										}
+									}							
+
+									if (const FScalarParameterNameAndCurve* ScalarCurve = ScalarCurves.FindByPredicate([Element](FScalarParameterNameAndCurve Curve)
+									{
+										return Curve.ParameterName == Element->GetName();
+									}))
+									{							
+										AnimSequencerHelpers::ConvertFloatChannelToRichCurve(ScalarCurve->ParameterCurve, FloatCurve.FloatCurve, FrameRate);
 									}
-								}							
-
-								if (const FScalarParameterNameAndCurve* ScalarCurve = ScalarCurves.FindByPredicate([Element](FScalarParameterNameAndCurve Curve)
-								{
-									return Curve.ParameterName == Element->GetName();
-								}))
-								{							
-									AnimSequencerHelpers::ConvertFloatChannelToRichCurve(ScalarCurve->ParameterCurve, FloatCurve.FloatCurve, FrameRate);
 								}
-							}
-							return true;
-						});	
-					}
-					else
-					{						
-						IAnimationDataController::ReportObjectErrorf(this, LOCTEXT("UnableToFindRigHierarchy", "Unable to retrieve RigHierarchy for ControlRig ({0})"), FText::FromString(ControlRig->GetPathName()));	      
-					}
+								return true;
+							});	
+						}
+						else
+						{						
+							IAnimationDataController::ReportObjectErrorf(this, LOCTEXT("UnableToFindRigHierarchy", "Unable to retrieve RigHierarchy for ControlRig ({0})"), FText::FromString(ControlRig->GetPathName()));	      
+						}
+					}								
 				}
 			}
 		}
 	}
-}
 
 void UAnimationSequencerDataModel::ValidateData() const
 {		
@@ -1035,55 +1039,55 @@ void UAnimationSequencerDataModel::IterateTransformControlCurve(const FName& Bon
 	const URigHierarchy* Hierarchy = ControlRig->GetHierarchy();
 	
 	const FRigElementKey BoneControlKey(UFKControlRig::GetControlName(BoneName, ERigElementType::Bone), ERigElementType::Control);
-	if (const FTransformParameterNameAndCurves* ControlCurvePtr = Section->GetTransformParameterNamesAndCurves().FindByPredicate([CurveName = BoneControlKey.Name](const FTransformParameterNameAndCurves& TransformParameter)
-	{
-		return TransformParameter.ParameterName == CurveName;
-	}))
-	{
-		const FTransformParameterNameAndCurves& ControlCurve = *ControlCurvePtr;
+		if (const FTransformParameterNameAndCurves* ControlCurvePtr = Section->GetTransformParameterNamesAndCurves().FindByPredicate([CurveName = BoneControlKey.Name](const FTransformParameterNameAndCurves& TransformParameter)
+		{
+			return TransformParameter.ParameterName == CurveName;
+		}))
+		{
+			const FTransformParameterNameAndCurves& ControlCurve = *ControlCurvePtr;
 
-		FVector3f Location(0.f);
-		FVector3f EulerAngles(0.f);
-		FVector3f Scale = FVector3f::OneVector;
+			FVector3f Location(0.f);
+			FVector3f EulerAngles(0.f);
+			FVector3f Scale = FVector3f::OneVector;
 
-		// Check whether or not any data is contained
-		bool bContainsData = false;
-		bool bContainsKeys = false;
-		for (int32 ChannelIndex = 0; ChannelIndex < 3; ++ChannelIndex)
-		{
-			bContainsData |= ControlCurve.Translation[ChannelIndex].HasAnyData();
-			bContainsKeys |= ControlCurve.Translation[ChannelIndex].GetNumKeys() != 0;
-			bContainsData |= ControlCurve.Rotation[ChannelIndex].HasAnyData();		
-			bContainsKeys |= ControlCurve.Rotation[ChannelIndex].GetNumKeys() != 0;		
-			bContainsData |= ControlCurve.Scale[ChannelIndex].HasAnyData();
-			bContainsKeys |= ControlCurve.Scale[ChannelIndex].GetNumKeys() != 0;
-		}
+			// Check whether or not any data is contained
+			bool bContainsData = false;
+			bool bContainsKeys = false;
+			for (int32 ChannelIndex = 0; ChannelIndex < 3; ++ChannelIndex)
+			{
+				bContainsData |= ControlCurve.Translation[ChannelIndex].HasAnyData();
+				bContainsKeys |= ControlCurve.Translation[ChannelIndex].GetNumKeys() != 0;
+				bContainsData |= ControlCurve.Rotation[ChannelIndex].HasAnyData();		
+				bContainsKeys |= ControlCurve.Rotation[ChannelIndex].GetNumKeys() != 0;		
+				bContainsData |= ControlCurve.Scale[ChannelIndex].HasAnyData();
+				bContainsKeys |= ControlCurve.Scale[ChannelIndex].GetNumKeys() != 0;
+			}
 
-		if (bContainsData)
-		{
-			const int32 NumberOfKeysToIterate = InFrameNumbers != nullptr ? InFrameNumbers->Num() : GetNumberOfKeys();
-			for (int32 KeyIndex = 0; KeyIndex < NumberOfKeysToIterate; ++KeyIndex)
+			if (bContainsData)
 			{
-				const FFrameNumber Frame = InFrameNumbers != nullptr ? (*InFrameNumbers)[KeyIndex] : FFrameNumber(KeyIndex);
-				for (int32 ChannelIndex = 0; ChannelIndex < 3; ++ChannelIndex)
+				const int32 NumberOfKeysToIterate = InFrameNumbers != nullptr ? InFrameNumbers->Num() : GetNumberOfKeys();
+				for (int32 KeyIndex = 0; KeyIndex < NumberOfKeysToIterate; ++KeyIndex)
 				{
-					ControlCurve.Translation[ChannelIndex].Evaluate(Frame, Location[ChannelIndex]);
-					ControlCurve.Rotation[ChannelIndex].Evaluate(Frame, EulerAngles[ChannelIndex]);
-					ControlCurve.Scale[ChannelIndex].Evaluate(Frame, Scale[ChannelIndex]);
-				}
+					const FFrameNumber Frame = InFrameNumbers != nullptr ? (*InFrameNumbers)[KeyIndex] : FFrameNumber(KeyIndex);
+					for (int32 ChannelIndex = 0; ChannelIndex < 3; ++ChannelIndex)
+					{
+						ControlCurve.Translation[ChannelIndex].Evaluate(Frame, Location[ChannelIndex]);
+						ControlCurve.Rotation[ChannelIndex].Evaluate(Frame, EulerAngles[ChannelIndex]);
+						ControlCurve.Scale[ChannelIndex].Evaluate(Frame, Scale[ChannelIndex]);
+					}
 
-				FTransform Transform;
-				Transform.SetLocation(FVector(Location));
-				Transform.SetRotation(FQuat::MakeFromEuler(FVector(EulerAngles)));
-				Transform.SetScale3D(FVector(Scale));
+					FTransform Transform;
+					Transform.SetLocation(FVector(Location));
+					Transform.SetRotation(FQuat::MakeFromEuler(FVector(EulerAngles)));
+					Transform.SetScale3D(FVector(Scale));
 
-				Transform.NormalizeRotation();
+					Transform.NormalizeRotation();
 
-				IterationFunction(Transform, Frame);
+					IterationFunction(Transform, Frame);
+				}
 			}
 		}
 	}
-}
 
 void UAnimationSequencerDataModel::GenerateTransformKeysForControl(const FName& BoneName, TArray<FTransform>& InOutTransforms, TArray<FFrameNumber>& InOutFrameNumbers) const
 {
@@ -1131,7 +1135,7 @@ void UAnimationSequencerDataModel::GeneratePoseData(UControlRig* ControlRig, FAn
 			Curve.Empty();
 
 			UE::Anim::Retargeting::FRetargetingScope RetargetingScope(GetSkeleton(), RigPose, EvaluationContext);
-
+			
 			// Populate bone/curve elements to Pose/Curve indices
 			{
 				QUICK_SCOPE_CYCLE_COUNTER(STAT_GetMappings);
@@ -1159,12 +1163,12 @@ void UAnimationSequencerDataModel::GeneratePoseData(UControlRig* ControlRig, FAn
 				});
 
 				RigHierarchy->ForEach<FRigCurveElement>([this, &RigHierarchy, &Curve](const FRigCurveElement* CurveElement) -> bool
-				{
-					const FName& CurveName = CurveElement->GetName();
+					{
+						const FName& CurveName = CurveElement->GetName();
 					Curve.Add(CurveName, RigHierarchy->GetCurveValue(CurveElement->GetKey()));
-					return true;
-				});
-			}
+						return true;
+					});
+				}
 
 			{
 				QUICK_SCOPE_CYCLE_COUNTER(STAT_NormalizeRotations);
@@ -1510,110 +1514,110 @@ void UAnimationSequencerDataModel::IterateBoneKeys(const FName& BoneName, TFunct
 	const URigHierarchy* Hierarchy = ControlRig->GetHierarchy();
 	
 	const FRigElementKey BoneControlKey(UFKControlRig::GetControlName(BoneName, ERigElementType::Bone), ERigElementType::Control);
-	if (const FTransformParameterNameAndCurves* ControlCurvePtr = Section->GetTransformParameterNamesAndCurves().FindByPredicate([CurveName = BoneControlKey.Name](const FTransformParameterNameAndCurves& TransformParameter)
-	{
-		return TransformParameter.ParameterName == CurveName;
-	}))
-	{
-		const FTransformParameterNameAndCurves& ControlCurve = *ControlCurvePtr;
-
-		struct FChannelInfo
+		if (const FTransformParameterNameAndCurves* ControlCurvePtr = Section->GetTransformParameterNamesAndCurves().FindByPredicate([CurveName = BoneControlKey.Name](const FTransformParameterNameAndCurves& TransformParameter)
 		{
-			bool bConstant = true;
-			bool bUniform = true;
-		};			
-		FChannelInfo PosChannels[3];
-		FChannelInfo RotChannels[3];
-		FChannelInfo ScaleChannels[3];
-
-		const int32 NumberOfKeys = GetNumberOfKeys();
-
-		int32 MaxNumberOfKeys = 1;
-		TSet<FFrameNumber> FrameNumbers;
-		for (int32 ChannelIndex = 0; ChannelIndex < 3; ++ChannelIndex)
+			return TransformParameter.ParameterName == CurveName;
+		}))
 		{
-			auto ProcessChannel = [&MaxNumberOfKeys, NumberOfKeys](const FMovieSceneFloatChannel& Channel, FChannelInfo& Info)
+			const FTransformParameterNameAndCurves& ControlCurve = *ControlCurvePtr;
+
+			struct FChannelInfo
+			{
+				bool bConstant = true;
+				bool bUniform = true;
+			};			
+			FChannelInfo PosChannels[3];
+			FChannelInfo RotChannels[3];
+			FChannelInfo ScaleChannels[3];
+
+			const int32 NumberOfKeys = GetNumberOfKeys();
+
+			int32 MaxNumberOfKeys = 1;
+			TSet<FFrameNumber> FrameNumbers;
+			for (int32 ChannelIndex = 0; ChannelIndex < 3; ++ChannelIndex)
 			{
-				if (Channel.HasAnyData())
+				auto ProcessChannel = [&MaxNumberOfKeys, NumberOfKeys](const FMovieSceneFloatChannel& Channel, FChannelInfo& Info)
 				{
-					if (Channel.GetNumKeys() == 0)
+					if (Channel.HasAnyData())
 					{
-						Info.bUniform = false;
-					}
-					else
-					{
-						Info.bConstant = false;
-						if (Channel.GetNumKeys() != NumberOfKeys)
+						if (Channel.GetNumKeys() == 0)
 						{
 							Info.bUniform = false;
 						}
+						else
+						{
+							Info.bConstant = false;
+							if (Channel.GetNumKeys() != NumberOfKeys)
+							{
+								Info.bUniform = false;
+							}
 
-						MaxNumberOfKeys = FMath::Max(MaxNumberOfKeys, Channel.GetNumKeys());
+							MaxNumberOfKeys = FMath::Max(MaxNumberOfKeys, Channel.GetNumKeys());
+						}
 					}
-				}
-			};
+				};
 
-			ProcessChannel(ControlCurve.Translation[ChannelIndex], PosChannels[ChannelIndex]);
-			ProcessChannel(ControlCurve.Rotation[ChannelIndex], RotChannels[ChannelIndex]);
-			ProcessChannel(ControlCurve.Scale[ChannelIndex], ScaleChannels[ChannelIndex]);
-		}
+				ProcessChannel(ControlCurve.Translation[ChannelIndex], PosChannels[ChannelIndex]);
+				ProcessChannel(ControlCurve.Rotation[ChannelIndex], RotChannels[ChannelIndex]);
+				ProcessChannel(ControlCurve.Scale[ChannelIndex], ScaleChannels[ChannelIndex]);
+			}
 
 
-		const int32 NumberOfKeysToIterate = MaxNumberOfKeys;
-		FVector3f PreviousPos, PreviousScale, PreviousRot;
-		FQuat4f PreviousQuat;
+			const int32 NumberOfKeysToIterate = MaxNumberOfKeys;
+			FVector3f PreviousPos, PreviousScale, PreviousRot;
+			FQuat4f PreviousQuat;
 
-		for (int32 KeyIndex = 0; KeyIndex < NumberOfKeysToIterate; ++KeyIndex)
-		{
-			const FFrameNumber Frame = KeyIndex;
-			
-			for (int32 ChannelIndex = 0; ChannelIndex < 3; ++ChannelIndex)
+			for (int32 KeyIndex = 0; KeyIndex < NumberOfKeysToIterate; ++KeyIndex)
 			{
-				if (PosChannels[ChannelIndex].bConstant)
+				const FFrameNumber Frame = KeyIndex;
+				
+				for (int32 ChannelIndex = 0; ChannelIndex < 3; ++ChannelIndex)
 				{
-					if (KeyIndex == 0)
+					if (PosChannels[ChannelIndex].bConstant)
 					{
-						PreviousPos[ChannelIndex] = ControlCurve.Translation[ChannelIndex].GetDefault().GetValue();
+						if (KeyIndex == 0)
+						{
+							PreviousPos[ChannelIndex] = ControlCurve.Translation[ChannelIndex].GetDefault().GetValue();
+						}
 					}
-				}
-				else
-				{
-					ControlCurve.Translation[ChannelIndex].Evaluate(Frame, PreviousPos[ChannelIndex]);
-				}
-				
-				if (RotChannels[ChannelIndex].bConstant)
-				{
-					if (KeyIndex == 0)
+					else
 					{
-						PreviousRot[ChannelIndex] = ControlCurve.Rotation[ChannelIndex].GetDefault().GetValue();
+						ControlCurve.Translation[ChannelIndex].Evaluate(Frame, PreviousPos[ChannelIndex]);
+					}
+					
+					if (RotChannels[ChannelIndex].bConstant)
+					{
+						if (KeyIndex == 0)
+						{
+							PreviousRot[ChannelIndex] = ControlCurve.Rotation[ChannelIndex].GetDefault().GetValue();
+						}
+					}
+					else
+					{
+						ControlCurve.Rotation[ChannelIndex].Evaluate(Frame, PreviousRot[ChannelIndex]);
 					}
-				}
-				else
-				{
-					ControlCurve.Rotation[ChannelIndex].Evaluate(Frame, PreviousRot[ChannelIndex]);
-				}
 
-				if (ScaleChannels[ChannelIndex].bConstant)
-				{
-					if (KeyIndex == 0)
+					if (ScaleChannels[ChannelIndex].bConstant)
+					{
+						if (KeyIndex == 0)
+						{
+							PreviousScale[ChannelIndex] = ControlCurve.Scale[ChannelIndex].GetDefault().GetValue();
+						}
+					}
+					else
 					{
-						PreviousScale[ChannelIndex] = ControlCurve.Scale[ChannelIndex].GetDefault().GetValue();
+						ControlCurve.Scale[ChannelIndex].Evaluate(Frame, PreviousScale[ChannelIndex]);
 					}
+					
 				}
-				else
+
+				PreviousQuat = FQuat4f::MakeFromEuler(PreviousRot);
+				if(!IterationFunction(PreviousPos, PreviousQuat, PreviousScale, Frame))
 				{
-					ControlCurve.Scale[ChannelIndex].Evaluate(Frame, PreviousScale[ChannelIndex]);
+					return;
 				}
-				
-			}
-
-			PreviousQuat = FQuat4f::MakeFromEuler(PreviousRot);
-			if(!IterationFunction(PreviousPos, PreviousQuat, PreviousScale, Frame))
-			{
-				return;
 			}
 		}
 	}
-}
 
 #undef LOCTEXT_NAMESPACE //"AnimSequencerDataModel"
\ No newline at end of file
diff --git a/Engine/Source/Runtime/Engine/Private/Animation/AnimDataModel.cpp b/Engine/Source/Runtime/Engine/Private/Animation/AnimDataModel.cpp
--- a/Engine/Source/Runtime/Engine/Private/Animation/AnimDataModel.cpp
+++ b/Engine/Source/Runtime/Engine/Private/Animation/AnimDataModel.cpp
@@ -548,7 +548,8 @@ FGuid UAnimDataModel::GenerateGuid() const
 
     for (const FTransformCurve& Curve : CurveData.TransformCurves)
     {
-    	UpdateWithString(Curve.GetName().ToString());
+    	const FString CurveName = Curve.GetName().ToString();
+    	UpdateSHAWithArray(CurveName.GetCharArray());
 
     	auto UpdateWithComponent = [&UpdateWithFloatCurve](const FVectorCurve& VectorCurve)
     	{
@@ -565,8 +566,11 @@ FGuid UAnimDataModel::GenerateGuid() const
 	
 	for (const FAnimatedBoneAttribute& Attribute : AnimatedBoneAttributes)
 	{
-		UpdateWithData(Attribute.Identifier);
-		const uint32 StructSize = Attribute.Identifier.GetType()->GetStructureSize();
+		UpdateSHAWithArray(Attribute.Identifier.GetName().ToString().GetCharArray());
+		UpdateSHAWithArray(Attribute.Identifier.GetBoneName().ToString().GetCharArray());
+		UpdateWithData(Attribute.Identifier.GetBoneIndex());
+		UpdateSHAWithArray(Attribute.Identifier.GetType()->GetFName().ToString().GetCharArray());
+		const uint32 StructSize = Attribute.Identifier.GetType()->GetPropertiesSize();
 		for (const FAttributeKey& Key : Attribute.Curve.GetConstRefOfKeys())
 		{
 			UpdateWithData(Key.Time);
