MH plugin export functionality can cause animation compression deadlock
#fix add modify/evaluate lock to top-level controller brackets - this as model notifies are delayed, which can cause model-derived data to be out-of-date
#misc change reaction to curve model changes (clear compressed binary data, only clear full cached data - which can cancel in-flight tasks - outside of brackets)

diff --git a/Engine/Plugins/Animation/AnimationData/Source/AnimationData/Private/AnimSequencerController.cpp b/Engine/Plugins/Animation/AnimationData/Source/AnimationData/Private/AnimSequencerController.cpp
--- a/Engine/Plugins/Animation/AnimationData/Source/AnimationData/Private/AnimSequencerController.cpp
+++ b/Engine/Plugins/Animation/AnimationData/Source/AnimationData/Private/AnimSequencerController.cpp
@@ -53,7 +53,7 @@ void UAnimSequencerController::OpenBracket(const FText& InTitle, bool bShouldTra
 	{
 		FBracketPayload Payload;
 		Payload.Description = InTitle.ToString();
-
+		BracketLock = MakeUnique<IAnimationDataModel::FEvaluationAndModificationLock>(*ModelInterface);
 		Model->GetNotifier().Notify(EAnimDataModelNotifyType::BracketOpened, Payload);
 	}
 
@@ -83,6 +83,7 @@ void UAnimSequencerController::CloseBracket(bool bShouldTransact /*= true*/)
 			ChangeTransactor.CloseTransaction();
 		}
 		
+		BracketLock.Reset();
 		Model->GetNotifier().Notify(EAnimDataModelNotifyType::BracketClosed);
 	}
 }
diff --git a/Engine/Plugins/Animation/AnimationData/Source/AnimationData/Private/AnimSequencerController.h b/Engine/Plugins/Animation/AnimationData/Source/AnimationData/Private/AnimSequencerController.h
--- a/Engine/Plugins/Animation/AnimationData/Source/AnimationData/Private/AnimSequencerController.h
+++ b/Engine/Plugins/Animation/AnimationData/Source/AnimationData/Private/AnimSequencerController.h
@@ -129,6 +129,8 @@ private:
 	UPROPERTY(transient)
 	TScriptInterface<IAnimationDataModel> ModelInterface;
 
+	TUniquePtr<IAnimationDataModel::FEvaluationAndModificationLock> BracketLock;
+
 	friend class FAnimDataControllerTestBase;
 	friend UE::Anim::FOpenBracketAction;
 	friend UE::Anim::FCloseBracketAction;
diff --git a/Engine/Source/Runtime/Engine/Private/Animation/AnimSequence.cpp b/Engine/Source/Runtime/Engine/Private/Animation/AnimSequence.cpp
--- a/Engine/Source/Runtime/Engine/Private/Animation/AnimSequence.cpp
+++ b/Engine/Source/Runtime/Engine/Private/Animation/AnimSequence.cpp
@@ -4837,10 +4837,11 @@ void UAnimSequence::OnModelModified(const EAnimDataModelNotifyType& NotifyType,
 		case EAnimDataModelNotifyType::CurveFlagsChanged:
 		case EAnimDataModelNotifyType::CurveScaled:
 		{
-			ClearAllCachedCookedPlatformData();
+			ClearCompressedCurveData();
 				
 			if (NotifyCollector.IsNotWithinBracket())
 			{						
+				ClearAllCachedCookedPlatformData();
 				UpdateRawDataGuid(RegenerateGUID);
 				BeginCacheDerivedDataForCurrentPlatform();
 			}
