Curve colors are lost when upgrading to AnimSequence data structure
#fix ensure that curve colors are kept when moving data to AnimSequencerModel

diff --git a/Engine/Plugins/Animation/AnimationData/Source/AnimationData/Private/AnimSequencerController.cpp b/Engine/Plugins/Animation/AnimationData/Source/AnimationData/Private/AnimSequencerController.cpp
--- a/Engine/Plugins/Animation/AnimationData/Source/AnimationData/Private/AnimSequencerController.cpp
+++ b/Engine/Plugins/Animation/AnimationData/Source/AnimationData/Private/AnimSequencerController.cpp
@@ -2287,6 +2287,7 @@ void UAnimSequencerController::PopulateWithExistingModel(TScriptInterface<IAnima
 						Attributes.SetPostExtrapolation(Curve.FloatCurve.PostInfinityExtrap);						
 						SetCurveAttributes(CurveId, Attributes, false);
 						SetCurveKeys(CurveId, Curve.FloatCurve.GetConstRefOfKeys(), false);
+						SetCurveColor(CurveId, Curve.GetColor(), false);
 					}
 				}
 				
