Can't record actor
#fix can only remap animation data is model has actually been populated with any

diff --git a/Engine/Plugins/Animation/AnimationData/Source/AnimationData/Private/AnimSequencerController.cpp b/Engine/Plugins/Animation/AnimationData/Source/AnimationData/Private/AnimSequencerController.cpp
--- a/Engine/Plugins/Animation/AnimationData/Source/AnimationData/Private/AnimSequencerController.cpp
+++ b/Engine/Plugins/Animation/AnimationData/Source/AnimationData/Private/AnimSequencerController.cpp
@@ -2196,21 +2196,24 @@ void UAnimSequencerController::RemoveUnusedControlsAndCurves() const
 
 void UAnimSequencerController::UpdateWithSkeleton(USkeleton* TargetSkeleton, bool bShouldTransact)
 {
-	OpenBracket(LOCTEXT("SettingNewskeleton", "Updating Skeleton for Animation Data Model"));
+	if (ModelInterface->HasBeenPopulated())
 	{
-		// (re-)generate the rig hierarchy
-		Model->InitializeFKControlRig(CastChecked<UFKControlRig>(Model->GetControlRig()), TargetSkeleton);
-		
-		// Remove any bone tracks that do not exist in the new hierarchy
-		RemoveBoneTracksMissingFromSkeleton(TargetSkeleton, bShouldTransact);
+		OpenBracket(LOCTEXT("SettingNewskeleton", "Updating Skeleton for Animation Data Model"));
+		{
+			// (re-)generate the rig hierarchy
+			Model->InitializeFKControlRig(CastChecked<UFKControlRig>(Model->GetControlRig()), TargetSkeleton);
 
-		// Remove any control/curves which were created in InitializeFKControlRig which are no longer keyed
-		RemoveUnusedControlsAndCurves();		
+			// Remove any bone tracks that do not exist in the new hierarchy
+			RemoveBoneTracksMissingFromSkeleton(TargetSkeleton, bShouldTransact);
 
-		// Forcefully re-generate legacy data structures
-		Model->GenerateLegacyCurveData();
-	}
-	CloseBracket();
+			// Remove any control/curves which were created in InitializeFKControlRig which are no longer keyed
+			RemoveUnusedControlsAndCurves();		
+
+			// Forcefully re-generate legacy data structures
+			Model->GenerateLegacyCurveData();
+		}
+		CloseBracket();
+	}	
 }
 
 void UAnimSequencerController::PopulateWithExistingModel(TScriptInterface<IAnimationDataModel> InModel)
