Truncation fixes for AnimationData

diff --git a/Engine/Plugins/Animation/AnimationData/Source/AnimationData/Private/AnimSequencerController.cpp b/Engine/Plugins/Animation/AnimationData/Source/AnimationData/Private/AnimSequencerController.cpp
--- a/Engine/Plugins/Animation/AnimationData/Source/AnimationData/Private/AnimSequencerController.cpp
+++ b/Engine/Plugins/Animation/AnimationData/Source/AnimationData/Private/AnimSequencerController.cpp
@@ -129,9 +129,9 @@ void UAnimSequencerController::ResizeNumberOfFrames(FFrameNumber NewLength, FFra
 					FSequenceLengthChangedPayload Payload;
 
 					PRAGMA_DISABLE_DEPRECATION_WARNINGS
-					Payload.T0 = CurrentFrameRate.AsSeconds(T0);
-					Payload.T1 = CurrentFrameRate.AsSeconds(T1);
-					Payload.PreviousLength = CurrentFrameRate.AsSeconds(CurrentNumberOfFrames);
+					Payload.T0 = static_cast<float>(CurrentFrameRate.AsSeconds(T0));
+					Payload.T1 = static_cast<float>(CurrentFrameRate.AsSeconds(T1));
+					Payload.PreviousLength = static_cast<float>(CurrentFrameRate.AsSeconds(CurrentNumberOfFrames));
 					PRAGMA_ENABLE_DEPRECATION_WARNINGS
 
 					Payload.PreviousNumberOfFrames = CurrentNumberOfFrames;
@@ -184,8 +184,8 @@ void UAnimSequencerController::ResizeInFrames(FFrameNumber NewLength, FFrameNumb
 					ResizeNumberOfFrames(NewLength, T0, T1, bShouldTransact);
 
 					const FFrameRate CurrentFrameRate = Model->GetMovieScene()->GetDisplayRate();
-					ResizeCurves(CurrentFrameRate.AsSeconds(NewLength), bInserted, CurrentFrameRate.AsSeconds(T0), CurrentFrameRate.AsSeconds(T1), bShouldTransact);
-					ResizeAttributes(CurrentFrameRate.AsSeconds(NewLength), bInserted, CurrentFrameRate.AsSeconds(T0), CurrentFrameRate.AsSeconds(T1), bShouldTransact);
+					ResizeCurves(static_cast<float>(CurrentFrameRate.AsSeconds(NewLength)), bInserted, static_cast<float>(CurrentFrameRate.AsSeconds(T0)), static_cast<float>(CurrentFrameRate.AsSeconds(T1)), bShouldTransact);
+					ResizeAttributes(static_cast<float>(CurrentFrameRate.AsSeconds(NewLength)), bInserted, static_cast<float>(CurrentFrameRate.AsSeconds(T0)), static_cast<float>(CurrentFrameRate.AsSeconds(T1)), bShouldTransact);
 				}
 				else
 				{
@@ -228,8 +228,8 @@ void UAnimSequencerController::SetMovieSceneRange(FFrameNumber InFrameNumber) co
 	FMovieSceneEditorData& EditorData = MovieScene->GetEditorData();
 	EditorData.ViewStart = EditorData.WorkStart = 0.f;
 	EditorData.ViewEnd = EditorData.WorkEnd = Model->GetFrameRate().AsSeconds(InFrameNumber);
-	MovieScene->SetWorkingRange(EditorData.ViewStart, EditorData.ViewEnd);
-	MovieScene->SetViewRange(EditorData.ViewStart, EditorData.ViewEnd);
+	MovieScene->SetWorkingRange(static_cast<float>(EditorData.ViewStart), static_cast<float>(EditorData.ViewEnd));
+	MovieScene->SetViewRange(static_cast<float>(EditorData.ViewStart), static_cast<float>(EditorData.ViewEnd));
 
 	if (UMovieSceneControlRigParameterSection* Section = Model->GetFKControlRigSection())
 	{
@@ -278,7 +278,7 @@ void UAnimSequencerController::SetFrameRate(FFrameRate FrameRate, bool bShouldTr
 	ValidateModel();
 
 	// Disallow invalid frame-rates, or 0.0 intervals
-	const float FrameRateInterval = FrameRate.AsInterval();
+	const float FrameRateInterval = static_cast<float>(FrameRate.AsInterval());
 	if (FrameRate.IsValid() && !FMath::IsNearlyZero(FrameRateInterval) && FrameRateInterval > 0.f)
 	{
 		const TObjectPtr<UMovieScene>& MovieScene = Model->MovieScene;
@@ -852,7 +852,7 @@ bool UAnimSequencerController::SetTransformCurveKeys(const FAnimationCurveIdenti
 				{
 					for (int32 ChannelIndex = 0; ChannelIndex < 3; ++ChannelIndex)
 					{
-						Key.ChannelKeys[ChannelIndex][KeyIndex] = FRichCurveKey(Time, Vector[ChannelIndex]);
+						Key.ChannelKeys[ChannelIndex][KeyIndex] = FRichCurveKey(Time, static_cast<float>(Vector[ChannelIndex]));
 					}
 				};
 
@@ -915,7 +915,7 @@ bool UAnimSequencerController::SetTransformCurveKey(const FAnimationCurveIdentif
 		{
 			for (int32 ChannelIndex = 0; ChannelIndex < 3; ++ChannelIndex)
 			{
-				Key.ChannelKeys[ChannelIndex] = FRichCurveKey(Time, Vector[ChannelIndex]);
+				Key.ChannelKeys[ChannelIndex] = FRichCurveKey(Time, static_cast<float>(Vector[ChannelIndex]));
 			}
 		};
 
@@ -1675,7 +1675,7 @@ bool UAnimSequencerController::UpdateBoneTrackKeys(FName BoneName, const FInt32R
 						for (int32 RangeIndex = RangeMin; RangeIndex < RangeMax; ++RangeIndex, ++KeyIndex)
 						{
 							TransformKeys.Add(FTransform(FQuat(RotationalKeys[KeyIndex]), FVector(PositionalKeys[KeyIndex]), FVector(ScalingKeys[KeyIndex])));
-							TimeValues.Add(FrameRate.AsSeconds(RangeIndex));
+							TimeValues.Add(static_cast<float>(FrameRate.AsSeconds(RangeIndex)));
 						}
 					
 						if(!UpdateBoneCurveKeys(BoneName, TransformKeys, TimeValues))
@@ -2834,13 +2834,13 @@ bool UAnimSequencerController::UpdateBoneCurveKeys(const FName& BoneName, const
 						for (int32 ChannelIndex = 0; ChannelIndex < 3; ++ChannelIndex)
 						{
 							// Translation
-							AddKeyToChannel(&ParameterCurvePair->Translation[ChannelIndex], FrameNumber, Location[ChannelIndex], EMovieSceneKeyInterpolation::Linear);
+							AddKeyToChannel(&ParameterCurvePair->Translation[ChannelIndex], FrameNumber, static_cast<float>(Location[ChannelIndex]), EMovieSceneKeyInterpolation::Linear);
 						
 							// Rotation
-							AddKeyToChannel(&ParameterCurvePair->Rotation[ChannelIndex], FrameNumber, EulerAngles[ChannelIndex], EMovieSceneKeyInterpolation::Linear);
+							AddKeyToChannel(&ParameterCurvePair->Rotation[ChannelIndex], FrameNumber, static_cast<float>(EulerAngles[ChannelIndex]), EMovieSceneKeyInterpolation::Linear);
 							
 							// Scaling							
-							AddKeyToChannel(&ParameterCurvePair->Scale[ChannelIndex], FrameNumber, Scale[ChannelIndex], EMovieSceneKeyInterpolation::Linear);						
+							AddKeyToChannel(&ParameterCurvePair->Scale[ChannelIndex], FrameNumber, static_cast<float>(Scale[ChannelIndex]), EMovieSceneKeyInterpolation::Linear);						
 						}
 					}
 					
diff --git a/Engine/Plugins/Animation/AnimationData/Source/AnimationData/Private/AnimSequencerDataModel.cpp b/Engine/Plugins/Animation/AnimationData/Source/AnimationData/Private/AnimSequencerDataModel.cpp
--- a/Engine/Plugins/Animation/AnimationData/Source/AnimationData/Private/AnimSequencerDataModel.cpp
+++ b/Engine/Plugins/Animation/AnimationData/Source/AnimationData/Private/AnimSequencerDataModel.cpp
@@ -1189,7 +1189,7 @@ void UAnimationSequencerDataModel::GeneratePoseData(UControlRig* ControlRig, FAn
 					// Add or retrieve curve
 					const FName& CurveName = TransformCurve.GetName();
 					// note we're not checking Curve.GetCurveTypeFlags() yet
-					FTransform Value = TransformCurve.Evaluate(EvaluationContext.SampleFrameRate.AsSeconds(EvaluationContext.SampleTime), 1.f);
+					FTransform Value = TransformCurve.Evaluate(static_cast<float>(EvaluationContext.SampleFrameRate.AsSeconds(EvaluationContext.SampleTime)), 1.f);
 
 					const FCompactPoseBoneIndex BoneIndex(RigPose.GetBoneContainer().GetPoseBoneIndexForBoneName(CurveName));
 					if(BoneIndex != INDEX_NONE)
@@ -1312,7 +1312,7 @@ void UAnimationSequencerDataModel::EvaluateTrack(UMovieSceneControlRigParameterT
 					{
 						FEulerTransform EulerTransform(FEulerTransform::Identity);
 					
-						const double Alpha = BoneSampleTime.GetSubFrame();
+						const float Alpha = BoneSampleTime.GetSubFrame();
 						auto EvaluateToTransform = [&TypedParameter](const FFrameNumber& Frame, FTransform& InOutTransform, TMovieSceneCurveChannelImpl<FMovieSceneFloatChannel>::FTimeEvaluationCache* Cache)
 						{
 							auto EvaluateValue = [Frame, Cache](const auto& Channel, auto& Target)
diff --git a/Engine/Plugins/Animation/AnimationData/Source/AnimationData/Private/AnimSequencerHelpers.cpp b/Engine/Plugins/Animation/AnimationData/Source/AnimationData/Private/AnimSequencerHelpers.cpp
--- a/Engine/Plugins/Animation/AnimationData/Source/AnimationData/Private/AnimSequencerHelpers.cpp
+++ b/Engine/Plugins/Animation/AnimationData/Source/AnimationData/Private/AnimSequencerHelpers.cpp
@@ -26,7 +26,7 @@ void AnimSequencerHelpers::ConvertFloatChannelToRichCurve(const FMovieSceneFloat
 		const FFrameTime ConvertedKeyFrameTime = FFrameRate::TransformTime(KeyFrameNumber, Channel.GetTickResolution(), TargetFrameRate);
 		const double KeyTime = TargetFrameRate.AsSeconds(ConvertedKeyFrameTime);
 
-		OutCurve.Keys[KeyIndex].Time = KeyTime;
+		OutCurve.Keys[KeyIndex].Time = static_cast<float>(KeyTime);
 
 		const FFrameNumber* PrevKey = KeyIndex > 0 ? &Times[KeyIndex - 1] : nullptr;
 		const FFrameNumber* NextKey = KeyIndex < (NumKeys - 1) ? &Times[KeyIndex + 1] : nullptr;
@@ -49,7 +49,7 @@ void AnimSequencerHelpers::ConvertFloatChannelToRichCurve(const FMovieSceneFloat
 			return 1;
 		}();
 		
-		const float SecondsDelta = [&]() -> float
+		const double SecondsDelta = [&]() -> double
 		{
 			if (PrevKey && NextKey)
 			{
@@ -64,7 +64,7 @@ void AnimSequencerHelpers::ConvertFloatChannelToRichCurve(const FMovieSceneFloat
 				return Channel.GetTickResolution().AsSeconds(*NextKey) - KeyTime;
 			}
 
-			return 1.f;
+			return 1.0;
 		}();		
 
 		// Ratio between rich-curve and moviescene key(s) timing, if there are any surrounding keys (otherwise default to a ratio of 1:1)
@@ -156,8 +156,8 @@ void AnimSequencerHelpers::ConvertRichCurveKeysToFloatChannel(const TArray<FRich
 				const double NextTimeDiff = FMath::Max<double>(KINDA_SMALL_NUMBER, Times[KeyIndex + 1].Value - Times[KeyIndex].Value);
 
 				// Generate linear tangent (delta value / delta time)
-				const float NewTangent = (NextKey.Value - Key.Value) / NextTimeDiff;
-				Key.Tangent.LeaveTangent = NewTangent;
+				const double NewTangent = (NextKey.Value - Key.Value) / NextTimeDiff;
+				Key.Tangent.LeaveTangent = static_cast<float>(NewTangent);
 			}
 		}
 		// Also need to generate linear tangents for cubic keys which are to the right of linear keys
@@ -175,7 +175,7 @@ void AnimSequencerHelpers::ConvertRichCurveKeysToFloatChannel(const TArray<FRich
 					
 					const double PrevTimeDiff = FMath::Max<double>(KINDA_SMALL_NUMBER, Times[KeyIndex].Value - Times[KeyIndex - 1].Value);
 					// Generate linear tangent (delta value / delta time)
-					Key.Tangent.ArriveTangent = (Key.Value - PreviousKey.Value) / PrevTimeDiff;
+					Key.Tangent.ArriveTangent = static_cast<float>((Key.Value - PreviousKey.Value) / PrevTimeDiff);
 				}
 			}
 		}
@@ -203,13 +203,13 @@ void AnimSequencerHelpers::ConvertRichCurveKeyToFloatValue(const FRichCurveKey&
 
 	if (OutMovieSceneKey.Tangent.TangentWeightMode == RCTWM_WeightedNone)
 	{
-		OutMovieSceneKey.Tangent.ArriveTangent = RichCurveKey.ArriveTangent * TangentRatio;
-		OutMovieSceneKey.Tangent.LeaveTangent = RichCurveKey.LeaveTangent * TangentRatio;
+		OutMovieSceneKey.Tangent.ArriveTangent = static_cast<float>(RichCurveKey.ArriveTangent * TangentRatio);
+		OutMovieSceneKey.Tangent.LeaveTangent = static_cast<float>(RichCurveKey.LeaveTangent * TangentRatio);
 	}
 	else
 	{		
-		OutMovieSceneKey.Tangent.ArriveTangent = RichCurveKey.ArriveTangent * SecondsPerFrame;
-		OutMovieSceneKey.Tangent.LeaveTangent = RichCurveKey.LeaveTangent * SecondsPerFrame;
+		OutMovieSceneKey.Tangent.ArriveTangent = static_cast<float>(RichCurveKey.ArriveTangent * SecondsPerFrame);
+		OutMovieSceneKey.Tangent.LeaveTangent = static_cast<float>(RichCurveKey.LeaveTangent * SecondsPerFrame);
 	}		
 			
 	OutMovieSceneKey.TangentMode = RichCurveKey.TangentMode;
@@ -222,8 +222,8 @@ void AnimSequencerHelpers::ConvertFloatValueToRichCurveKey(const FMovieSceneFloa
 	OutRichCurveKey.InterpMode = MovieSceneKey.InterpMode;	
 	
 	OutRichCurveKey.TangentMode = MovieSceneKey.TangentMode;
-	OutRichCurveKey.ArriveTangent = MovieSceneKey.Tangent.ArriveTangent * TangentRatio;
-	OutRichCurveKey.LeaveTangent = MovieSceneKey.Tangent.LeaveTangent * TangentRatio;
+	OutRichCurveKey.ArriveTangent = static_cast<float>(MovieSceneKey.Tangent.ArriveTangent * TangentRatio);
+	OutRichCurveKey.LeaveTangent = static_cast<float>(MovieSceneKey.Tangent.LeaveTangent * TangentRatio);
 		
 	OutRichCurveKey.TangentWeightMode = MovieSceneKey.Tangent.TangentWeightMode;
 	OutRichCurveKey.ArriveTangentWeight = MovieSceneKey.Tangent.ArriveTangentWeight;
