Move soft-lock tracking to be opt-in at compile time and gated by ini-settings for specific files/hashes

diff --git a/Engine/Plugins/Animation/AnimationData/Source/AnimationData/Private/AnimSequencerDataModel.cpp b/Engine/Plugins/Animation/AnimationData/Source/AnimationData/Private/AnimSequencerDataModel.cpp
--- a/Engine/Plugins/Animation/AnimationData/Source/AnimationData/Private/AnimSequencerDataModel.cpp
+++ b/Engine/Plugins/Animation/AnimationData/Source/AnimationData/Private/AnimSequencerDataModel.cpp
@@ -23,8 +23,9 @@
 #include "Compilation/MovieSceneCompiledDataManager.h"
 #include "UObject/UObjectThreadContext.h"
 #include "Animation/AnimCurveTypes.h"
-
 #include "Runtime/MovieScene/Private/Channels/MovieSceneCurveChannelImpl.h"
+#include "Containers/StackTracker.h"
+#include "AnimationCompression.h"
 
 #define LOCTEXT_NAMESPACE "AnimSequencerDataModel"
 
@@ -185,6 +186,18 @@ void UAnimationSequencerDataModel::PostLoad()
 
 		CachedRawDataGUID.Invalidate();
 
+#if UE_DEBUG_ANIMATION_COMPRESSION				
+		const bool bOutputDebugData = UAnimationSettings::Get()->CompressionDebugAssetPaths.Contains(GetOuter()->GetPathName());
+		if (bOutputDebugData)
+		{
+			UE_LOG(LogAnimationCompression, Display, TEXT("CachedRawDataGUID.Invalidate %s"), *GetPathName());
+			FStackTracker TempTracker(nullptr, nullptr, nullptr, true);
+			TempTracker.CaptureStackTrace(1);
+			TempTracker.DumpStackTraces(0, *GLog);
+			TempTracker.ResetTracking();
+		}
+#endif // UE_DEBUG_ANIMATION_COMPRESSION
+
 		RemoveOutOfDateControls();
 
 		ValidateData();
@@ -573,9 +586,10 @@ UAnimSequence* UAnimationSequencerDataModel::GetAnimationSequence() const
 
 FGuid UAnimationSequencerDataModel::GenerateGuid() const
 {	
+	FGuid ReturnGuid;
 	if (CachedRawDataGUID.IsValid())
 	{
-		return CachedRawDataGUID;
+		ReturnGuid = CachedRawDataGUID;
 	}
 	else
 	{
@@ -676,8 +690,23 @@ FGuid UAnimationSequencerDataModel::GenerateGuid() const
 		Sha.GetHash(reinterpret_cast<uint8*>(Hash));
 		const FGuid Guid(Hash[0] ^ Hash[4], Hash[1], Hash[2], Hash[3]);
 		
-		return Guid;
+		ReturnGuid = Guid;
 	}
+
+
+#if UE_DEBUG_ANIMATION_COMPRESSION
+	const bool bOutputDebugData = UAnimationSettings::Get()->CompressionDebugGUIDs.Contains(ReturnGuid) || UAnimationSettings::Get()->CompressionDebugAssetPaths.Contains(GetOuter()->GetPathName());
+	if (bOutputDebugData)
+	{
+		UE_LOG(LogAnimationCompression, Display, TEXT("UAnimationSequencerDataModel::GenerateGuid %s - %s"), *GetPathName(), *ReturnGuid.ToString());
+		FStackTracker TempTracker(nullptr, nullptr, nullptr, true);
+		TempTracker.CaptureStackTrace(1);
+		TempTracker.DumpStackTraces(0, *GLog);
+		TempTracker.ResetTracking();
+	}
+#endif // UE_DEBUG_ANIMATION_COMPRESSION
+
+	return ReturnGuid;
 }
 
 TScriptInterface<IAnimationDataController> UAnimationSequencerDataModel::GetController()
@@ -735,6 +764,17 @@ void UAnimationSequencerDataModel::OnNotify(const EAnimDataModelNotifyType& Noti
 			if (CachedRawDataGUID.IsValid() && (!Collector.Contains(EAnimDataModelNotifyType::Populated) || !FUObjectThreadContext::Get().IsRoutingPostLoad))
 			{
 				CachedRawDataGUID.Invalidate();			
+#if UE_DEBUG_ANIMATION_COMPRESSION
+				const bool bOutputDebugData = UAnimationSettings::Get()->CompressionDebugAssetPaths.Contains(GetOuter()->GetPathName());
+				if (bOutputDebugData)
+				{
+					UE_LOG(LogAnimationCompression, Display, TEXT("CachedRawDataGUID.Invalidate %s"), *GetPathName());
+					FStackTracker TempTracker(nullptr, nullptr, nullptr, true);
+					TempTracker.CaptureStackTrace(1);
+					TempTracker.DumpStackTraces(0, *GLog);
+					TempTracker.ResetTracking();
+				}
+#endif // UE_DEBUG_ANIMATION_COMPRESSION
 			}
 		};
 
@@ -1039,55 +1079,55 @@ void UAnimationSequencerDataModel::IterateTransformControlCurve(const FName& Bon
 	const URigHierarchy* Hierarchy = ControlRig->GetHierarchy();
 	
 	const FRigElementKey BoneControlKey(UFKControlRig::GetControlName(BoneName, ERigElementType::Bone), ERigElementType::Control);
-		if (const FTransformParameterNameAndCurves* ControlCurvePtr = Section->GetTransformParameterNamesAndCurves().FindByPredicate([CurveName = BoneControlKey.Name](const FTransformParameterNameAndCurves& TransformParameter)
-		{
-			return TransformParameter.ParameterName == CurveName;
-		}))
-		{
-			const FTransformParameterNameAndCurves& ControlCurve = *ControlCurvePtr;
+	if (const FTransformParameterNameAndCurves* ControlCurvePtr = Section->GetTransformParameterNamesAndCurves().FindByPredicate([CurveName = BoneControlKey.Name](const FTransformParameterNameAndCurves& TransformParameter)
+	{
+		return TransformParameter.ParameterName == CurveName;
+	}))
+	{
+		const FTransformParameterNameAndCurves& ControlCurve = *ControlCurvePtr;
 
-			FVector3f Location(0.f);
-			FVector3f EulerAngles(0.f);
-			FVector3f Scale = FVector3f::OneVector;
+		FVector3f Location(0.f);
+		FVector3f EulerAngles(0.f);
+		FVector3f Scale = FVector3f::OneVector;
 
-			// Check whether or not any data is contained
-			bool bContainsData = false;
-			bool bContainsKeys = false;
-			for (int32 ChannelIndex = 0; ChannelIndex < 3; ++ChannelIndex)
-			{
-				bContainsData |= ControlCurve.Translation[ChannelIndex].HasAnyData();
-				bContainsKeys |= ControlCurve.Translation[ChannelIndex].GetNumKeys() != 0;
-				bContainsData |= ControlCurve.Rotation[ChannelIndex].HasAnyData();		
-				bContainsKeys |= ControlCurve.Rotation[ChannelIndex].GetNumKeys() != 0;		
-				bContainsData |= ControlCurve.Scale[ChannelIndex].HasAnyData();
-				bContainsKeys |= ControlCurve.Scale[ChannelIndex].GetNumKeys() != 0;
-			}
+		// Check whether or not any data is contained
+		bool bContainsData = false;
+		bool bContainsKeys = false;
+		for (int32 ChannelIndex = 0; ChannelIndex < 3; ++ChannelIndex)
+		{
+			bContainsData |= ControlCurve.Translation[ChannelIndex].HasAnyData();
+			bContainsKeys |= ControlCurve.Translation[ChannelIndex].GetNumKeys() != 0;
+			bContainsData |= ControlCurve.Rotation[ChannelIndex].HasAnyData();		
+			bContainsKeys |= ControlCurve.Rotation[ChannelIndex].GetNumKeys() != 0;		
+			bContainsData |= ControlCurve.Scale[ChannelIndex].HasAnyData();
+			bContainsKeys |= ControlCurve.Scale[ChannelIndex].GetNumKeys() != 0;
+		}
 
-			if (bContainsData)
+		if (bContainsData)
+		{
+			const int32 NumberOfKeysToIterate = InFrameNumbers != nullptr ? InFrameNumbers->Num() : GetNumberOfKeys();
+			for (int32 KeyIndex = 0; KeyIndex < NumberOfKeysToIterate; ++KeyIndex)
 			{
-				const int32 NumberOfKeysToIterate = InFrameNumbers != nullptr ? InFrameNumbers->Num() : GetNumberOfKeys();
-				for (int32 KeyIndex = 0; KeyIndex < NumberOfKeysToIterate; ++KeyIndex)
+				const FFrameNumber Frame = InFrameNumbers != nullptr ? (*InFrameNumbers)[KeyIndex] : FFrameNumber(KeyIndex);
+				for (int32 ChannelIndex = 0; ChannelIndex < 3; ++ChannelIndex)
 				{
-					const FFrameNumber Frame = InFrameNumbers != nullptr ? (*InFrameNumbers)[KeyIndex] : FFrameNumber(KeyIndex);
-					for (int32 ChannelIndex = 0; ChannelIndex < 3; ++ChannelIndex)
-					{
-						ControlCurve.Translation[ChannelIndex].Evaluate(Frame, Location[ChannelIndex]);
-						ControlCurve.Rotation[ChannelIndex].Evaluate(Frame, EulerAngles[ChannelIndex]);
-						ControlCurve.Scale[ChannelIndex].Evaluate(Frame, Scale[ChannelIndex]);
-					}
+					ControlCurve.Translation[ChannelIndex].Evaluate(Frame, Location[ChannelIndex]);
+					ControlCurve.Rotation[ChannelIndex].Evaluate(Frame, EulerAngles[ChannelIndex]);
+					ControlCurve.Scale[ChannelIndex].Evaluate(Frame, Scale[ChannelIndex]);
+				}
 
-					FTransform Transform;
-					Transform.SetLocation(FVector(Location));
-					Transform.SetRotation(FQuat::MakeFromEuler(FVector(EulerAngles)));
-					Transform.SetScale3D(FVector(Scale));
+				FTransform Transform;
+				Transform.SetLocation(FVector(Location));
+				Transform.SetRotation(FQuat::MakeFromEuler(FVector(EulerAngles)));
+				Transform.SetScale3D(FVector(Scale));
 
-					Transform.NormalizeRotation();
+				Transform.NormalizeRotation();
 
-					IterationFunction(Transform, Frame);
-				}
+				IterationFunction(Transform, Frame);
 			}
 		}
 	}
+}
 
 void UAnimationSequencerDataModel::GenerateTransformKeysForControl(const FName& BoneName, TArray<FTransform>& InOutTransforms, TArray<FFrameNumber>& InOutFrameNumbers) const
 {
@@ -1514,110 +1554,110 @@ void UAnimationSequencerDataModel::IterateBoneKeys(const FName& BoneName, TFunct
 	const URigHierarchy* Hierarchy = ControlRig->GetHierarchy();
 	
 	const FRigElementKey BoneControlKey(UFKControlRig::GetControlName(BoneName, ERigElementType::Bone), ERigElementType::Control);
-		if (const FTransformParameterNameAndCurves* ControlCurvePtr = Section->GetTransformParameterNamesAndCurves().FindByPredicate([CurveName = BoneControlKey.Name](const FTransformParameterNameAndCurves& TransformParameter)
+	if (const FTransformParameterNameAndCurves* ControlCurvePtr = Section->GetTransformParameterNamesAndCurves().FindByPredicate([CurveName = BoneControlKey.Name](const FTransformParameterNameAndCurves& TransformParameter)
+	{
+		return TransformParameter.ParameterName == CurveName;
+	}))
+	{
+		const FTransformParameterNameAndCurves& ControlCurve = *ControlCurvePtr;
+
+		struct FChannelInfo
 		{
-			return TransformParameter.ParameterName == CurveName;
-		}))
+			bool bConstant = true;
+			bool bUniform = true;
+		};			
+		FChannelInfo PosChannels[3];
+		FChannelInfo RotChannels[3];
+		FChannelInfo ScaleChannels[3];
+
+		const int32 NumberOfKeys = GetNumberOfKeys();
+
+		int32 MaxNumberOfKeys = 1;
+		TSet<FFrameNumber> FrameNumbers;
+		for (int32 ChannelIndex = 0; ChannelIndex < 3; ++ChannelIndex)
 		{
-			const FTransformParameterNameAndCurves& ControlCurve = *ControlCurvePtr;
-
-			struct FChannelInfo
-			{
-				bool bConstant = true;
-				bool bUniform = true;
-			};			
-			FChannelInfo PosChannels[3];
-			FChannelInfo RotChannels[3];
-			FChannelInfo ScaleChannels[3];
-
-			const int32 NumberOfKeys = GetNumberOfKeys();
-
-			int32 MaxNumberOfKeys = 1;
-			TSet<FFrameNumber> FrameNumbers;
-			for (int32 ChannelIndex = 0; ChannelIndex < 3; ++ChannelIndex)
+			auto ProcessChannel = [&MaxNumberOfKeys, NumberOfKeys](const FMovieSceneFloatChannel& Channel, FChannelInfo& Info)
 			{
-				auto ProcessChannel = [&MaxNumberOfKeys, NumberOfKeys](const FMovieSceneFloatChannel& Channel, FChannelInfo& Info)
+				if (Channel.HasAnyData())
 				{
-					if (Channel.HasAnyData())
+					if (Channel.GetNumKeys() == 0)
+					{
+						Info.bUniform = false;
+					}
+					else
 					{
-						if (Channel.GetNumKeys() == 0)
+						Info.bConstant = false;
+						if (Channel.GetNumKeys() != NumberOfKeys)
 						{
 							Info.bUniform = false;
 						}
-						else
-						{
-							Info.bConstant = false;
-							if (Channel.GetNumKeys() != NumberOfKeys)
-							{
-								Info.bUniform = false;
-							}
 
-							MaxNumberOfKeys = FMath::Max(MaxNumberOfKeys, Channel.GetNumKeys());
-						}
+						MaxNumberOfKeys = FMath::Max(MaxNumberOfKeys, Channel.GetNumKeys());
 					}
-				};
+				}
+			};
 
-				ProcessChannel(ControlCurve.Translation[ChannelIndex], PosChannels[ChannelIndex]);
-				ProcessChannel(ControlCurve.Rotation[ChannelIndex], RotChannels[ChannelIndex]);
-				ProcessChannel(ControlCurve.Scale[ChannelIndex], ScaleChannels[ChannelIndex]);
-			}
+			ProcessChannel(ControlCurve.Translation[ChannelIndex], PosChannels[ChannelIndex]);
+			ProcessChannel(ControlCurve.Rotation[ChannelIndex], RotChannels[ChannelIndex]);
+			ProcessChannel(ControlCurve.Scale[ChannelIndex], ScaleChannels[ChannelIndex]);
+		}
 
 
-			const int32 NumberOfKeysToIterate = MaxNumberOfKeys;
-			FVector3f PreviousPos, PreviousScale, PreviousRot;
-			FQuat4f PreviousQuat;
+		const int32 NumberOfKeysToIterate = MaxNumberOfKeys;
+		FVector3f PreviousPos, PreviousScale, PreviousRot;
+		FQuat4f PreviousQuat;
 
-			for (int32 KeyIndex = 0; KeyIndex < NumberOfKeysToIterate; ++KeyIndex)
+		for (int32 KeyIndex = 0; KeyIndex < NumberOfKeysToIterate; ++KeyIndex)
+		{
+			const FFrameNumber Frame = KeyIndex;
+			
+			for (int32 ChannelIndex = 0; ChannelIndex < 3; ++ChannelIndex)
 			{
-				const FFrameNumber Frame = KeyIndex;
-				
-				for (int32 ChannelIndex = 0; ChannelIndex < 3; ++ChannelIndex)
+				if (PosChannels[ChannelIndex].bConstant)
 				{
-					if (PosChannels[ChannelIndex].bConstant)
+					if (KeyIndex == 0)
 					{
-						if (KeyIndex == 0)
-						{
-							PreviousPos[ChannelIndex] = ControlCurve.Translation[ChannelIndex].GetDefault().GetValue();
-						}
-					}
-					else
-					{
-						ControlCurve.Translation[ChannelIndex].Evaluate(Frame, PreviousPos[ChannelIndex]);
+						PreviousPos[ChannelIndex] = ControlCurve.Translation[ChannelIndex].GetDefault().GetValue();
 					}
-					
-					if (RotChannels[ChannelIndex].bConstant)
-					{
-						if (KeyIndex == 0)
-						{
-							PreviousRot[ChannelIndex] = ControlCurve.Rotation[ChannelIndex].GetDefault().GetValue();
-						}
-					}
-					else
+				}
+				else
+				{
+					ControlCurve.Translation[ChannelIndex].Evaluate(Frame, PreviousPos[ChannelIndex]);
+				}
+				
+				if (RotChannels[ChannelIndex].bConstant)
+				{
+					if (KeyIndex == 0)
 					{
-						ControlCurve.Rotation[ChannelIndex].Evaluate(Frame, PreviousRot[ChannelIndex]);
+						PreviousRot[ChannelIndex] = ControlCurve.Rotation[ChannelIndex].GetDefault().GetValue();
 					}
+				}
+				else
+				{
+					ControlCurve.Rotation[ChannelIndex].Evaluate(Frame, PreviousRot[ChannelIndex]);
+				}
 
-					if (ScaleChannels[ChannelIndex].bConstant)
-					{
-						if (KeyIndex == 0)
-						{
-							PreviousScale[ChannelIndex] = ControlCurve.Scale[ChannelIndex].GetDefault().GetValue();
-						}
-					}
-					else
+				if (ScaleChannels[ChannelIndex].bConstant)
+				{
+					if (KeyIndex == 0)
 					{
-						ControlCurve.Scale[ChannelIndex].Evaluate(Frame, PreviousScale[ChannelIndex]);
+						PreviousScale[ChannelIndex] = ControlCurve.Scale[ChannelIndex].GetDefault().GetValue();
 					}
-					
 				}
-
-				PreviousQuat = FQuat4f::MakeFromEuler(PreviousRot);
-				if(!IterationFunction(PreviousPos, PreviousQuat, PreviousScale, Frame))
+				else
 				{
-					return;
+					ControlCurve.Scale[ChannelIndex].Evaluate(Frame, PreviousScale[ChannelIndex]);
 				}
+				
+			}
+
+			PreviousQuat = FQuat4f::MakeFromEuler(PreviousRot);
+			if(!IterationFunction(PreviousPos, PreviousQuat, PreviousScale, Frame))
+			{
+				return;
 			}
 		}
 	}
+}
 
 #undef LOCTEXT_NAMESPACE //"AnimSequencerDataModel"
\ No newline at end of file
diff --git a/Engine/Source/Runtime/Engine/Classes/Animation/AnimSequence.h b/Engine/Source/Runtime/Engine/Classes/Animation/AnimSequence.h
--- a/Engine/Source/Runtime/Engine/Classes/Animation/AnimSequence.h
+++ b/Engine/Source/Runtime/Engine/Classes/Animation/AnimSequence.h
@@ -27,8 +27,9 @@
 #include "AnimData/IAnimationDataModel.h"
 #endif // WITH_EDITOR
 
-#include "AnimSequence.generated.h"
+#define UE_DEBUG_ANIMATION_COMPRESSION 1
 
+#include "AnimSequence.generated.h"
 
 typedef TArray<FTransform> FTransformArrayA2;
 
@@ -882,12 +883,12 @@ protected:
 	bool PollCacheDerivedData(const FIoHash& KeyHash) const;
 	void EndCacheDerivedData(const FIoHash& KeyHash);
 
-	void LogCompressionRequestTimings() const;
-
 	FIoHash DataKeyHash;
 	TMap<FIoHash, TUniquePtr<FCompressedAnimSequence>> DataByPlatformKeyHash;
 	TMap<FIoHash, TPimplPtr<UE::Anim::FAnimationSequenceAsyncCacheTask>> CacheTasksByKeyHash;
 
+#if UE_DEBUG_ANIMATION_COMPRESSION
+	void LogCompressionRequestTimings() const;
 	struct FCompressionRequestTimings
 	{
 		FDateTime LastBeginCacheRequest;
@@ -897,6 +898,7 @@ protected:
 	};
 	TMap<FIoHash, FCompressionRequestTimings> TimingsByKeyHash;
 	TMap<FIoHash, FString> DDCStringByKeyHash;
+#endif // UE_DEBUG_ANIMATION_COMPRESSION
 protected:
 	bool TryCancelAsyncTasks();
 	void FinishAsyncTasks();
diff --git a/Engine/Source/Runtime/Engine/Classes/Animation/AnimationSettings.h b/Engine/Source/Runtime/Engine/Classes/Animation/AnimationSettings.h
--- a/Engine/Source/Runtime/Engine/Classes/Animation/AnimationSettings.h
+++ b/Engine/Source/Runtime/Engine/Classes/Animation/AnimationSettings.h
@@ -104,6 +104,19 @@ class ENGINE_API UAnimationSettings : public UDeveloperSettings
 	/** Project specific default frame-rate used when (re)initializing any animation based data */
 	UPROPERTY(config, EditAnywhere, Category = AnimationData)
 	FFrameRate DefaultFrameRate;
+
+	/** Data model GUIDs which should be tracked by animation compression debugging/logging */
+	UPROPERTY(config, EditAnywhere, Category = AnimationData)
+	TArray<FGuid> CompressionDebugGUIDs;
+
+	/** Animation sequence DDC hashes which should be tracked by animation compression debugging/logging */
+	UPROPERTY(config, EditAnywhere, Category = AnimationData)
+	TArray<FString> CompressionDebugHashes;
+
+	/** Animation sequence full paths which should be tracked by animation compression debugging/logging */
+	UPROPERTY(config, EditAnywhere, Category = AnimationData)
+	TArray<FString> CompressionDebugAssetPaths;
+	
 public:
 	static UAnimationSettings * Get() { return CastChecked<UAnimationSettings>(UAnimationSettings::StaticClass()->GetDefaultObject()); }
 
diff --git a/Engine/Source/Runtime/Engine/Private/Animation/AnimSequence.cpp b/Engine/Source/Runtime/Engine/Private/Animation/AnimSequence.cpp
--- a/Engine/Source/Runtime/Engine/Private/Animation/AnimSequence.cpp
+++ b/Engine/Source/Runtime/Engine/Private/Animation/AnimSequence.cpp
@@ -452,6 +452,19 @@ void UAnimSequence::WillNeverCacheCookedPlatformDataAgain()
 {
 	Super::WillNeverCacheCookedPlatformDataAgain();
 
+#if UE_DEBUG_ANIMATION_COMPRESSION
+	const bool bOutputDebugData = UAnimationSettings::Get()->CompressionDebugAssetPaths.Contains(GetPathName());
+	if (bOutputDebugData)
+	{
+		UE_LOG(LogAnimationCompression, Display, TEXT("UAnimSequence::WillNeverCacheCookedPlatformDataAgain %s"), *GetPathName());
+
+		for (TPair<FIoHash, FCompressionRequestTimings>& TimingsPair : TimingsByKeyHash)
+		{
+			TimingsPair.Value.LastWillNeverCook = FDateTime::Now();	
+		}
+	}
+#endif // UE_DEBUG_ANIMATION_COMPRESSION
+
 	UE::Anim::FAnimSequenceCompilingManager::Get().FinishCompilation({this});	
 	ClearCompressedCurveData();
 	ClearCompressedBoneData();
@@ -459,11 +472,6 @@ void UAnimSequence::WillNeverCacheCookedPlatformDataAgain()
 	// Clear out current platform, and any target platform data
 	CompressedData.Reset();
 	
-	for (TPair<FIoHash, FCompressionRequestTimings>& TimingsPair : TimingsByKeyHash)
-	{
-		TimingsPair.Value.LastWillNeverCook = FDateTime::Now();	
-	}
-		
 	CacheTasksByKeyHash.Empty();
 	DataByPlatformKeyHash.Empty();
 	PRAGMA_DISABLE_DEPRECATION_WARNINGS
@@ -475,17 +483,24 @@ void UAnimSequence::ClearAllCachedCookedPlatformData()
 {
 	Super::ClearAllCachedCookedPlatformData();
 	
+#if UE_DEBUG_ANIMATION_COMPRESSION
+	const bool bOutputDebugData = UAnimationSettings::Get()->CompressionDebugAssetPaths.Contains(GetPathName());
+	if (bOutputDebugData)
+	{
+		UE_LOG(LogAnimationCompression, Display, TEXT("UAnimSequence::ClearAllCachedCookedPlatformData %s"), *GetPathName());
+		for (TPair<FIoHash, FCompressionRequestTimings>& TimingsPair : TimingsByKeyHash)
+		{
+			TimingsPair.Value.LastClearCache = FDateTime::Now();	
+		}
+	}
+#endif // UE_DEBUG_ANIMATION_COMPRESSION
+	
 	// Delete any cache tasks first because the destructor will cancel the cache and build tasks,
 	// and drop their pointers to the data.
 	CacheTasksByKeyHash.Empty();
 	DataByPlatformKeyHash.Empty();
 	CompressedData.Reset();
 	DataKeyHash = FIoHash::Zero;
-
-	for (TPair<FIoHash, FCompressionRequestTimings>& TimingsPair : TimingsByKeyHash)
-	{
-		TimingsPair.Value.LastClearCache = FDateTime::Now();	
-	}
 	
 	PRAGMA_DISABLE_DEPRECATION_WARNINGS
 	bUseRawDataOnly = true;
@@ -1102,10 +1117,20 @@ bool UAnimSequence::IsCachedCookedPlatformDataLoaded(const ITargetPlatform* Targ
 
 	auto LogCookedDataWarning = [this, TargetPlatform, KeyHash]()
 	{
-		// This means cooked data is expected, but it it currently invalid and no task is running to generate it
-		UE_LOG(LogAnimation, Warning, TEXT("Expecting either valid compressed Animation Data, or an in-flight compression task for asset %s\nPlatform: %s\nHash: %s"), *GetPathName(), *TargetPlatform->DisplayName().ToString(), *LexToString(KeyHash));
+#if UE_DEBUG_ANIMATION_COMPRESSION
+		const bool bOutputDebugData = UAnimationSettings::Get()->CompressionDebugHashes.Contains(LexToString(KeyHash)) || UAnimationSettings::Get()->CompressionDebugAssetPaths.Contains(GetPathName());
+		if (bOutputDebugData)
+		{
+			// This means cooked data is expected, but it it currently invalid and no task is running to generate it
+			UE_LOG(LogAnimationCompression, Warning, TEXT("Expecting either valid compressed Animation Data, or an in-flight compression task for asset %s\nPlatform: %s\nHash: %s"), *GetPathName(), *TargetPlatform->DisplayName().ToString(), *LexToString(KeyHash));
 
+			if (bOutputDebugData)
+			{
+				UE_LOG(LogAnimationCompression, Display, TEXT("UAnimSequence::IsCachedCookedPlatformDataLoaded %s - %s"), *GetPathName(), *LexToString(KeyHash));
+			}
+		}
 		LogCompressionRequestTimings();
+#endif // UE_DEBUG_ANIMATION_COMPRESSION
 
 		if (IsInGameThread())
 		{
@@ -1123,7 +1148,6 @@ bool UAnimSequence::IsCachedCookedPlatformDataLoaded(const ITargetPlatform* Targ
 		else if (!CacheTasksByKeyHash.Contains(KeyHash))
 		{
 			LogCookedDataWarning();
-			CompressedData.IsValid(this, true);
 		}
 	}
 	else
@@ -1137,7 +1161,6 @@ bool UAnimSequence::IsCachedCookedPlatformDataLoaded(const ITargetPlatform* Targ
 			else if (!CacheTasksByKeyHash.Contains(KeyHash))
 			{
 				LogCookedDataWarning();
-				(*CompressedDataPtr)->IsValid(this, true);
 			}
 		}
 		else if (!CacheTasksByKeyHash.Contains(KeyHash))
@@ -5049,7 +5072,14 @@ FIoHash UAnimSequence::CreateDerivedDataKeyHash(const ITargetPlatform* TargetPla
 	Writer << Ret;
 
 	const FIoHash FinalHash = Writer.Finalize();
-	DDCStringByKeyHash.FindOrAdd(FinalHash) = Ret;
+
+#if UE_DEBUG_ANIMATION_COMPRESSION
+	const bool bOutputDebugData = UAnimationSettings::Get()->CompressionDebugHashes.Contains(LexToString(FinalHash)) || UAnimationSettings::Get()->CompressionDebugAssetPaths.Contains(GetPathName());
+	if (bOutputDebugData)
+	{
+		DDCStringByKeyHash.FindOrAdd(FinalHash) = Ret;
+	}
+#endif // UE_DEBUG_ANIMATION_COMPRESSION
 	return FinalHash;
 }
 
@@ -5108,8 +5138,15 @@ FIoHash UAnimSequence::BeginCacheDerivedData(const ITargetPlatform* TargetPlatfo
     {
     	TargetData = DataByPlatformKeyHash.Emplace(KeyHash, MakeUnique<FCompressedAnimSequence>()).Get();
     }
-	
-	TimingsByKeyHash.FindOrAdd(KeyHash).LastBeginCacheRequest = FDateTime::Now();
+
+#if UE_DEBUG_ANIMATION_COMPRESSION
+	const bool bOutputDebugData = UAnimationSettings::Get()->CompressionDebugHashes.Contains(LexToString(KeyHash)) || UAnimationSettings::Get()->CompressionDebugAssetPaths.Contains(GetPathName());
+	if (bOutputDebugData)
+	{
+		UE_LOG(LogAnimationCompression, Display, TEXT("UAnimSequence::BeginCacheDerivedData %s - %s"), *GetPathName(), *LexToString(KeyHash));
+		TimingsByKeyHash.FindOrAdd(KeyHash).LastBeginCacheRequest = FDateTime::Now();
+	}
+#endif // UE_DEBUG_ANIMATION_COMPRESSION
 	check(TargetData);
 	TargetData->Reset();
 
@@ -5178,10 +5215,19 @@ void UAnimSequence::EndCacheDerivedData(const FIoHash& KeyHash)
 	{
 		return;
 	}
+
+#if UE_DEBUG_ANIMATION_COMPRESSION
+	const bool bOutputDebugData = UAnimationSettings::Get()->CompressionDebugHashes.Contains(LexToString(KeyHash)) || UAnimationSettings::Get()->CompressionDebugAssetPaths.Contains(GetPathName());
+	if (bOutputDebugData)
+	{
+		UE_LOG(LogAnimationCompression, Display, TEXT("UAnimSequence::EndCacheDerivedData %s - %s"), *GetPathName(), *LexToString(KeyHash));
+	}
+#endif // UE_DEBUG_ANIMATION_COMPRESSION
 	
 	UE::Anim::FAnimSequenceCompilingManager::Get().FinishCompilation({this});
 }
 
+#if UE_DEBUG_ANIMATION_COMPRESSION
 void UAnimSequence::LogCompressionRequestTimings() const
 {
 	for (const TPair<FIoHash, FCompressionRequestTimings>& TimingsPair : TimingsByKeyHash)
@@ -5201,6 +5247,7 @@ void UAnimSequence::LogCompressionRequestTimings() const
 		*StringPair.Value);
 	}
 }
+#endif // UE_DEBUG_ANIMATION_COMPRESSION
 
 FCompressedAnimSequence& UAnimSequence::CacheDerivedData(const ITargetPlatform* TargetPlatform)
 {
@@ -5308,7 +5355,14 @@ void UAnimSequence::FinishAsyncTasks()
                     PRAGMA_ENABLE_DEPRECATION_WARNINGS
 				}
 
-				TimingsByKeyHash.FindChecked(It->Key).LastApplyCacheData = FDateTime::Now();
+#if UE_DEBUG_ANIMATION_COMPRESSION
+				const bool bOutputDebugData = UAnimationSettings::Get()->CompressionDebugHashes.Contains(LexToString(It->Key)) || UAnimationSettings::Get()->CompressionDebugAssetPaths.Contains(GetPathName());
+				if (bOutputDebugData)
+				{
+					UE_LOG(LogAnimationCompression, Display, TEXT("UAnimSequence::FinishAsyncTasks %s - %s"), *GetPathName(), *LexToString(It->Key));					
+					TimingsByKeyHash.FindChecked(It->Key).LastApplyCacheData = FDateTime::Now();
+				}
+#endif // UE_DEBUG_ANIMATION_COMPRESSION
 			}
 			else
 			{
