Fix disappearing curves, scalar parameter entry is used for (re)populating FK rig hierarchy during PostLoad

diff --git a/Engine/Plugins/Animation/AnimationData/Source/AnimationData/Private/AnimSequencerController.cpp b/Engine/Plugins/Animation/AnimationData/Source/AnimationData/Private/AnimSequencerController.cpp
--- a/Engine/Plugins/Animation/AnimationData/Source/AnimationData/Private/AnimSequencerController.cpp
+++ b/Engine/Plugins/Animation/AnimationData/Source/AnimationData/Private/AnimSequencerController.cpp
@@ -2967,7 +2967,7 @@ bool UAnimSequencerController::RemoveBoneCurveKey(const FName& BoneName, float T
 
 bool UAnimSequencerController::AddCurveControl(const FName& CurveName) const
 {
-	if (const UMovieSceneControlRigParameterSection* Section = Model->GetFKControlRigSection())
+	if (UMovieSceneControlRigParameterSection* Section = Model->GetFKControlRigSection())
 	{
 		UControlRig* ControlRig = Section->GetControlRig();
 		if (UFKControlRig* FKRig = Cast<UFKControlRig>(ControlRig))
@@ -2986,6 +2986,8 @@ bool UAnimSequencerController::AddCurveControl(const FName& CurveName) const
 					{
 						HierarchyController->AddCurve(CurveKey.Name, 0.f, false);		
 					}
+
+					Section->AddScalarParameter(CurveControlKey.Name, TOptional<float>(), true);
 				
 					const FRigCurveElement* CurveElement = Hierarchy->FindChecked<FRigCurveElement>(CurveKey);
 					ensure(CurveElement);
