Animation Sequence: Do not recreate CR cache when added unless model is populated and never when curve is added when populated.

diff --git a/Engine/Plugins/Animation/AnimationData/Source/AnimationData/Private/AnimSequencerController.cpp b/Engine/Plugins/Animation/AnimationData/Source/AnimationData/Private/AnimSequencerController.cpp
--- a/Engine/Plugins/Animation/AnimationData/Source/AnimationData/Private/AnimSequencerController.cpp
+++ b/Engine/Plugins/Animation/AnimationData/Source/AnimationData/Private/AnimSequencerController.cpp
@@ -2987,7 +2987,7 @@ bool UAnimSequencerController::AddCurveControl(const FName& CurveName) const
 						HierarchyController->AddCurve(CurveKey.Name, 0.f, false);		
 					}
 
-					Section->AddScalarParameter(CurveControlKey.Name, TOptional<float>(), true);
+					Section->AddScalarParameter(CurveControlKey.Name, TOptional<float>(), false);
 				
 					const FRigCurveElement* CurveElement = Hierarchy->FindChecked<FRigCurveElement>(CurveKey);
 					ensure(CurveElement);
@@ -3292,7 +3292,7 @@ bool UAnimSequencerController::SetCurveControlKey(const FName& CurveName, const
 					const bool bHasCurveControlChannel = Section->HasScalarParameter(CurveControlKey.Name);
 					if(!bHasCurveControlChannel)
 					{
-						Section->AddScalarParameter(CurveControlKey.Name, TOptional<float>(), true);
+						Section->AddScalarParameter(CurveControlKey.Name, TOptional<float>(), Model->bPopulated);
 					}
 
 					FScalarParameterNameAndCurve* ParameterCurvePair = Section->GetScalarParameterNamesAndCurves().FindByPredicate([CurveControlKey](const FScalarParameterNameAndCurve& Parameter)
