Fix crash when undoing/redoing animation sequence edits while an async task is in flight

Guards the PreEditUndo/PostEditUndo bracket so that any transaction modifications cannot occur at the same time as async compression/evaluation

diff --git a/Engine/Plugins/Animation/AnimationData/Source/AnimationData/Private/AnimSequencerDataModel.cpp b/Engine/Plugins/Animation/AnimationData/Source/AnimationData/Private/AnimSequencerDataModel.cpp
--- a/Engine/Plugins/Animation/AnimationData/Source/AnimationData/Private/AnimSequencerDataModel.cpp
+++ b/Engine/Plugins/Animation/AnimationData/Source/AnimationData/Private/AnimSequencerDataModel.cpp
@@ -224,6 +224,26 @@ void UAnimationSequencerDataModel::WillNeverCacheCookedPlatformDataAgain()
 {
 	Super::WillNeverCacheCookedPlatformDataAgain();
 }
+
+void UAnimationSequencerDataModel::PreEditUndo()
+{
+	Super::PreEditUndo();
+
+	// Lock evaluation as underlying MovieScene will be modified by undo/redo. Async compression tasks
+	// will be kicked off post-transaction of the model, but the underlying MovieScene may be 
+	// transacted after the model and modified concurrently with a compression task.
+	// We can do this because PreEditUndo calls are called on all objects in a transaction prior to 
+	// its application
+	LockEvaluationAndModification();
+}
+
+void UAnimationSequencerDataModel::PostEditUndo()
+{
+	// Unlock evaluation to allow for compression/evaluation now modifications are complete
+	UnlockEvaluationAndModification();
+
+	Super::PostEditUndo();
+}
 #endif
 
 double UAnimationSequencerDataModel::GetPlayLength() const
diff --git a/Engine/Plugins/Animation/AnimationData/Source/AnimationData/Private/AnimSequencerDataModel.h b/Engine/Plugins/Animation/AnimationData/Source/AnimationData/Private/AnimSequencerDataModel.h
--- a/Engine/Plugins/Animation/AnimationData/Source/AnimationData/Private/AnimSequencerDataModel.h
+++ b/Engine/Plugins/Animation/AnimationData/Source/AnimationData/Private/AnimSequencerDataModel.h
@@ -43,6 +43,8 @@ public:
 	virtual bool IsEditorOnly() const override { return true; }
 #if WITH_EDITOR
 	virtual void WillNeverCacheCookedPlatformDataAgain() override;
+	virtual void PreEditUndo() override;
+	virtual void PostEditUndo() override;
 #endif
 	/** End UObject overrides */
 
