Speculative fix for crash during automated test

All other accesses to GetControlRigTrack() are guarded by a null check, this adds one to GetFKControlrigSection()

diff --git a/Engine/Plugins/Animation/AnimationData/Source/AnimationData/Private/AnimSequencerDataModel.cpp b/Engine/Plugins/Animation/AnimationData/Source/AnimationData/Private/AnimSequencerDataModel.cpp
--- a/Engine/Plugins/Animation/AnimationData/Source/AnimationData/Private/AnimSequencerDataModel.cpp
+++ b/Engine/Plugins/Animation/AnimationData/Source/AnimationData/Private/AnimSequencerDataModel.cpp
@@ -820,16 +820,18 @@ UMovieSceneControlRigParameterSection* UAnimationSequencerDataModel::GetFKContro
 {
 	if (MovieScene)
 	{
-		const UMovieSceneControlRigParameterTrack* Track = GetControlRigTrack();
-		for (UMovieSceneSection* TrackSection : Track->GetAllSections())
+		if(const UMovieSceneControlRigParameterTrack* Track = GetControlRigTrack())
 		{
-			if (UMovieSceneControlRigParameterSection* Section = Cast<UMovieSceneControlRigParameterSection>(TrackSection))
+			for (UMovieSceneSection* TrackSection : Track->GetAllSections())
 			{
-				if (const UControlRig* ControlRig = Section->GetControlRig())
+				if (UMovieSceneControlRigParameterSection* Section = Cast<UMovieSceneControlRigParameterSection>(TrackSection))
 				{
-					if (ControlRig->IsA<UFKControlRig>())
+					if (const UControlRig* ControlRig = Section->GetControlRig())
 					{
-						return Section;
+						if (ControlRig->IsA<UFKControlRig>())
+						{
+							return Section;
+						}
 					}
 				}
 			}
