Fix for crash when adding curves to anim sequences and immediately getting the bone pose

diff --git a/Engine/Plugins/Animation/AnimationData/Source/AnimationData/Private/AnimSequencerController.cpp b/Engine/Plugins/Animation/AnimationData/Source/AnimationData/Private/AnimSequencerController.cpp
--- a/Engine/Plugins/Animation/AnimationData/Source/AnimationData/Private/AnimSequencerController.cpp
+++ b/Engine/Plugins/Animation/AnimationData/Source/AnimationData/Private/AnimSequencerController.cpp
@@ -3002,12 +3002,16 @@ bool UAnimSequencerController::AddCurveControl(const FName& CurveName) const
 						const FRigControlElement* ControlElement = Hierarchy->FindChecked<FRigControlElement>(CurveControlKey);
 						ensure(ControlElement);
 						
+						FKRig->RefreshActiveControls();
+						
 						return true;
 					}
 					else
 					{					
 						ReportWarningf(LOCTEXT("FailedToAddCurvecontrol", "Failed to add Curve control with name {0} as it already exists"), FText::FromName(CurveControlKey.Name));			
 					}
+
+					FKRig->RefreshActiveControls();
 				}
 				else
 				{
