AnimSequencerDataModel: handle certain curve/bone changes directly - forcing an update on the FKRig for the cached (control) data

diff --git a/Engine/Plugins/Animation/AnimationData/Source/AnimationData/Private/AnimSequencerController.cpp b/Engine/Plugins/Animation/AnimationData/Source/AnimationData/Private/AnimSequencerController.cpp
--- a/Engine/Plugins/Animation/AnimationData/Source/AnimationData/Private/AnimSequencerController.cpp
+++ b/Engine/Plugins/Animation/AnimationData/Source/AnimationData/Private/AnimSequencerController.cpp
@@ -3002,16 +3002,12 @@ bool UAnimSequencerController::AddCurveControl(const FName& CurveName) const
 						const FRigControlElement* ControlElement = Hierarchy->FindChecked<FRigControlElement>(CurveControlKey);
 						ensure(ControlElement);
 						
-						FKRig->RefreshActiveControls();
-						
 						return true;
 					}
 					else
 					{					
 						ReportWarningf(LOCTEXT("FailedToAddCurvecontrol", "Failed to add Curve control with name {0} as it already exists"), FText::FromName(CurveControlKey.Name));			
 					}
-
-					FKRig->RefreshActiveControls();
 				}
 				else
 				{
diff --git a/Engine/Plugins/Animation/AnimationData/Source/AnimationData/Private/AnimSequencerDataModel.cpp b/Engine/Plugins/Animation/AnimationData/Source/AnimationData/Private/AnimSequencerDataModel.cpp
--- a/Engine/Plugins/Animation/AnimationData/Source/AnimationData/Private/AnimSequencerDataModel.cpp
+++ b/Engine/Plugins/Animation/AnimationData/Source/AnimationData/Private/AnimSequencerDataModel.cpp
@@ -750,7 +750,7 @@ void UAnimationSequencerDataModel::OnNotify(const EAnimDataModelNotifyType& Noti
 {
 	Collector.Handle(NotifyType);
 
-	if (Collector.IsNotWithinBracket() && bPopulated)
+	if (bPopulated)
 	{
 		// Once the model has been populated and a modification is made - invalidate the cached GUID
 		auto ResetCachedGUID = [this]()
@@ -782,29 +782,41 @@ void UAnimationSequencerDataModel::OnNotify(const EAnimDataModelNotifyType& Noti
 				bRefreshed = true;
 			}
         };	
-		
-		const TArray<EAnimDataModelNotifyType> CurveNotifyTypes = {EAnimDataModelNotifyType::CurveAdded, EAnimDataModelNotifyType::CurveChanged, EAnimDataModelNotifyType::CurveRenamed, EAnimDataModelNotifyType::CurveRemoved,
-			EAnimDataModelNotifyType::CurveFlagsChanged, EAnimDataModelNotifyType::CurveScaled, EAnimDataModelNotifyType::CurveColorChanged, EAnimDataModelNotifyType::Populated, EAnimDataModelNotifyType::Reset };
-		if(Collector.Contains(CurveNotifyTypes))
+
+		if (Collector.IsNotWithinBracket())
 		{
-			if(!ValidationMode)
+			const TArray<EAnimDataModelNotifyType> CurveNotifyTypes = {EAnimDataModelNotifyType::CurveAdded, EAnimDataModelNotifyType::CurveChanged, EAnimDataModelNotifyType::CurveRenamed, EAnimDataModelNotifyType::CurveRemoved,
+			EAnimDataModelNotifyType::CurveFlagsChanged, EAnimDataModelNotifyType::CurveScaled, EAnimDataModelNotifyType::CurveColorChanged, EAnimDataModelNotifyType::Populated, EAnimDataModelNotifyType::Reset };
+			if(Collector.Contains(CurveNotifyTypes))
 			{
-				GenerateLegacyCurveData();
+				if(!ValidationMode)
+				{
+					GenerateLegacyCurveData();
+				}
+				RefreshControlsAndProxy();
+				ResetCachedGUID();
 			}
-			RefreshControlsAndProxy();
-			ResetCachedGUID();
-		}
 
-		const TArray<EAnimDataModelNotifyType> BonesNotifyTypes = {EAnimDataModelNotifyType::TrackAdded, EAnimDataModelNotifyType::TrackChanged, EAnimDataModelNotifyType::TrackRemoved, EAnimDataModelNotifyType::Populated, EAnimDataModelNotifyType::Reset };
-		if(Collector.Contains(BonesNotifyTypes))
-		{
-			RefreshControlsAndProxy();
-			ResetCachedGUID();
-		}
+			const TArray<EAnimDataModelNotifyType> BonesNotifyTypes = {EAnimDataModelNotifyType::TrackAdded, EAnimDataModelNotifyType::TrackChanged, EAnimDataModelNotifyType::TrackRemoved, EAnimDataModelNotifyType::Populated, EAnimDataModelNotifyType::Reset };
+			if(Collector.Contains(BonesNotifyTypes))
+			{
+				RefreshControlsAndProxy();
+				ResetCachedGUID();
+			}
 
-		if (Collector.Contains(EAnimDataModelNotifyType::Populated))
+			if (Collector.Contains(EAnimDataModelNotifyType::Populated))
+			{
+				RefreshControlsAndProxy();
+			}
+		}
+		else
 		{
-			RefreshControlsAndProxy();
+			// These changes can cause subsequent evaluation to fail due to mismatching data (related to changed controls)
+			const TArray<EAnimDataModelNotifyType> RigModificationTypes = {EAnimDataModelNotifyType::TrackAdded, EAnimDataModelNotifyType::TrackRemoved, EAnimDataModelNotifyType::CurveAdded, EAnimDataModelNotifyType::CurveRenamed, EAnimDataModelNotifyType::CurveRemoved};
+			if(Collector.Contains(RigModificationTypes))
+			{
+				RefreshControlsAndProxy();
+			}
 		}
 				
 		ValidateData();
