Different approach for setting initial random color for newly added float curves
#rnx

diff --git a/Engine/Plugins/Animation/AnimationData/Source/AnimationData/Private/AnimSequencerController.cpp b/Engine/Plugins/Animation/AnimationData/Source/AnimationData/Private/AnimSequencerController.cpp
--- a/Engine/Plugins/Animation/AnimationData/Source/AnimationData/Private/AnimSequencerController.cpp
+++ b/Engine/Plugins/Animation/AnimationData/Source/AnimationData/Private/AnimSequencerController.cpp
@@ -484,7 +484,6 @@ bool UAnimSequencerController::AddCurve(const FAnimationCurveIdentifier& CurveId
 					Model->CurveIdentifierToMetaData.FindOrAdd(CurveId).Flags = CurveFlags;
 				};
 
-				ConditionalAction<UE::Anim::FRemoveCurveAction>(bShouldTransact,  CurveId);
 				switch (CurveId.CurveType)
 				{
 				case ERawCurveTrackTypes::RCT_Transform:
@@ -496,7 +495,7 @@ bool UAnimSequencerController::AddCurve(const FAnimationCurveIdentifier& CurveId
 					{
 						ReportError(LOCTEXT("FailedtoAddCurveControl", "Failed to add curve control"));
 					}
-					SetCurveColor(CurveId, FAnimCurveBase::MakeColor(CurveId.CurveName), bShouldTransact);
+					Model->CurveIdentifierToMetaData.FindChecked(CurveId).Color = FAnimCurveBase::MakeColor(CurveId.CurveName);
 					break;
 				default:
 					{
@@ -505,6 +504,7 @@ bool UAnimSequencerController::AddCurve(const FAnimationCurveIdentifier& CurveId
 					}
 				}
 
+				ConditionalAction<UE::Anim::FRemoveCurveAction>(bShouldTransact,  CurveId);
 				Model->GetNotifier().Notify(EAnimDataModelNotifyType::CurveAdded, Payload);
 
 				return true;
