ExportAnimSequence tries to re-initialize animation data model

diff --git a/Engine/Plugins/Animation/AnimationData/Source/AnimationData/Private/AnimSequencerController.cpp b/Engine/Plugins/Animation/AnimationData/Source/AnimationData/Private/AnimSequencerController.cpp
--- a/Engine/Plugins/Animation/AnimationData/Source/AnimationData/Private/AnimSequencerController.cpp
+++ b/Engine/Plugins/Animation/AnimationData/Source/AnimationData/Private/AnimSequencerController.cpp
@@ -2425,17 +2425,16 @@ void UAnimSequencerController::InitializeModel()
 					{
 						ReportError(LOCTEXT("FailedToCreateControlRig", "Failed to create new UFKControlRig"));
 					}
-					
 				}
 				else
 				{
 					ReportError(LOCTEXT("FailedToCreateTrack", "Failed to create UMovieSceneControlRigParameterTrack"));
 				}
 			}
-		}
-		else
-		{
-			ReportError(LOCTEXT("FailedToCreateMovieScene", "Failed to create new UMovieScene"));
+			else
+			{
+				ReportError(LOCTEXT("FailedToCreateMovieScene", "Failed to create new UMovieScene"));
+			}
 		}
 	}
 	else
