Float Curve Keys added via AddFloatCurveKey are no longer frame aligned in 5.2

diff --git a/Engine/Plugins/Animation/AnimationData/Source/AnimationData/Private/AnimSequencerController.cpp b/Engine/Plugins/Animation/AnimationData/Source/AnimationData/Private/AnimSequencerController.cpp
--- a/Engine/Plugins/Animation/AnimationData/Source/AnimationData/Private/AnimSequencerController.cpp
+++ b/Engine/Plugins/Animation/AnimationData/Source/AnimationData/Private/AnimSequencerController.cpp
@@ -3308,11 +3308,24 @@ bool UAnimSequencerController::SetCurveControlKey(const FName& CurveName, const
 					
 					FMovieSceneFloatValue MovieSceneValue;
 					AnimSequencerHelpers::ConvertRichCurveKeyToFloatValue(Key, MovieSceneValue );
-					const FFrameNumber FrameNumber = Model->MovieScene->GetTickResolution().AsFrameNumber(Key.Time);
-					const FFrameNumber CurveFrameNumber = ConvertFrameTime(FrameNumber, Model->MovieScene->GetTickResolution(), ParameterCurvePair->ParameterCurve.GetTickResolution()).FrameNumber;
+
+					/*
+					 * Scenarios for curve keys conversion
+					 *
+					 * 1. All keys align with frame boundaries -> add all and convert between key types
+					 * 2. Some keys do not align with frame boundaries -> resample keys to MaxChannelRate (240000fps) to minimize evaluated value deltas
+					 */
+					const FFrameRate CurveRate = ParameterCurvePair->ParameterCurve.GetTickResolution();
+					const FFrameTime CurveFrameTime = CurveRate.AsFrameTime(Key.Time);
+					const bool bOffFrameCurveKey = !(FMath::IsNearlyZero(CurveFrameTime.GetSubFrame(), KINDA_SMALL_NUMBER) || FMath::IsNearlyEqual(CurveFrameTime.GetSubFrame(), 1.0f, KINDA_SMALL_NUMBER));
+					if(bOffFrameCurveKey)
+					{
+						const FFrameRate MaxChannelRate = FFrameRate(240000, 1);
+						ParameterCurvePair->ParameterCurve.ChangeFrameResolution(CurveRate, MaxChannelRate);
+						ParameterCurvePair->ParameterCurve.SetTickResolution(MaxChannelRate);
+					}					
 							
-					const FKeyHandle KeyHandle = ParameterCurvePair->ParameterCurve.GetData().UpdateOrAddKey(CurveFrameNumber, MovieSceneValue);
-										
+					const FKeyHandle KeyHandle = ParameterCurvePair->ParameterCurve.GetData().UpdateOrAddKey(ParameterCurvePair->ParameterCurve.GetTickResolution().AsFrameNumber(Key.Time), MovieSceneValue);
 					return KeyHandle != FKeyHandle::Invalid();
 				}
 				else
diff --git a/Engine/Plugins/Animation/AnimationData/Source/AnimationData/Private/AnimSequencerHelpers.cpp b/Engine/Plugins/Animation/AnimationData/Source/AnimationData/Private/AnimSequencerHelpers.cpp
--- a/Engine/Plugins/Animation/AnimationData/Source/AnimationData/Private/AnimSequencerHelpers.cpp
+++ b/Engine/Plugins/Animation/AnimationData/Source/AnimationData/Private/AnimSequencerHelpers.cpp
@@ -23,8 +23,7 @@ void AnimSequencerHelpers::ConvertFloatChannelToRichCurve(const FMovieSceneFloat
 	{
 		// Convert from channel to requested frame-rate
 		const FFrameNumber KeyFrameNumber = Times[KeyIndex];
-		const FFrameTime ConvertedKeyFrameTime = FFrameRate::TransformTime(KeyFrameNumber, Channel.GetTickResolution(), TargetFrameRate);
-		const double KeyTime = TargetFrameRate.AsSeconds(ConvertedKeyFrameTime);
+		const double KeyTime = Channel.GetTickResolution().AsSeconds(KeyFrameNumber);
 
 		OutCurve.Keys[KeyIndex].Time = static_cast<float>(KeyTime);
 
