Fix for incorrect frame number calculation when updating/adding anim curve key
#fix ensure that newly added movie scene curve always match their tick resolution to the MovieScene
#misc added additional validation to ensure requested action is performed correctly

diff --git a/Engine/Plugins/Animation/AnimationData/Source/AnimationData/Private/AnimSequencerController.cpp b/Engine/Plugins/Animation/AnimationData/Source/AnimationData/Private/AnimSequencerController.cpp
--- a/Engine/Plugins/Animation/AnimationData/Source/AnimationData/Private/AnimSequencerController.cpp
+++ b/Engine/Plugins/Animation/AnimationData/Source/AnimationData/Private/AnimSequencerController.cpp
@@ -1145,9 +1145,9 @@ bool UAnimSequencerController::SetCurveKey(const FAnimationCurveIdentifier& Curv
 		FCurveChangedPayload Payload;
 		Payload.Identifier = CurveId;
 
-		auto SetControlCurve = [this, CurveId, Key]()
+		auto SetControlCurve = [this, CurveId, Key](bool bUpdate)
 		{
-			if(CurveId.CurveType == ERawCurveTrackTypes::RCT_Float && !SetCurveControlKey(CurveId.CurveName, Key))
+			if(CurveId.CurveType == ERawCurveTrackTypes::RCT_Float && !SetCurveControlKey(CurveId.CurveName, Key, bUpdate))
 			{
 				ReportError(LOCTEXT("FailedtoSetCurveControlKey", "Failed to set curve control key"));
 			}	
@@ -1163,7 +1163,7 @@ bool UAnimSequencerController::SetCurveKey(const FAnimationCurveIdentifier& Curv
 			const FRichCurveKey CurrentKey = RichCurve->GetKey(Handle);
 			ConditionalAction<UE::Anim::FSetRichCurveKeyAction>(bShouldTransact,  CurveId, CurrentKey);
 
-			SetControlCurve();
+			SetControlCurve(true);
 
 			// Set the new value
 			RichCurve->SetKeyValue(Handle, Key.Value);
@@ -1175,7 +1175,7 @@ bool UAnimSequencerController::SetCurveKey(const FAnimationCurveIdentifier& Curv
 			FTransaction Transaction = ConditionalTransaction(LOCTEXT("AddNamedCurveKey", "Adding Curve Key"), bShouldTransact);
 			ConditionalAction<UE::Anim::FRemoveRichCurveKeyAction>(bShouldTransact,  CurveId, Key.Time);
 
-			SetControlCurve();
+			SetControlCurve(false);
 
 			// Add the new key
 			RichCurve->AddKey(Key.Time, Key.Value);
@@ -2987,7 +2987,22 @@ bool UAnimSequencerController::AddCurveControl(const FName& CurveName) const
 						HierarchyController->AddCurve(CurveKey.Name, 0.f, false);		
 					}
 
-					Section->AddScalarParameter(CurveControlKey.Name, TOptional<float>(), false);
+					const bool bHasCurveControlChannel = Section->HasScalarParameter(CurveControlKey.Name);
+					if (!bHasCurveControlChannel)
+					{
+						Section->AddScalarParameter(CurveControlKey.Name, TOptional<float>(), false);
+					}
+
+					FScalarParameterNameAndCurve* ParameterCurvePair = Section->GetScalarParameterNamesAndCurves().FindByPredicate([CurveControlKey](const FScalarParameterNameAndCurve& Parameter)
+					{
+						return Parameter.ParameterName == CurveControlKey.Name;
+					});
+					check(ParameterCurvePair);
+
+					if(!bHasCurveControlChannel)
+					{
+						ParameterCurvePair->ParameterCurve.SetTickResolution(Model->GetFrameRate());
+					}
 				
 					const FRigCurveElement* CurveElement = Hierarchy->FindChecked<FRigCurveElement>(CurveKey);
 					ensure(CurveElement);
@@ -3272,7 +3287,7 @@ bool UAnimSequencerController::SetCurveControlKeys(const FName& CurveName, const
 	return false;
 }
 
-bool UAnimSequencerController::SetCurveControlKey(const FName& CurveName, const FRichCurveKey& Key) const
+bool UAnimSequencerController::SetCurveControlKey(const FName& CurveName, const FRichCurveKey& Key, bool bUpdateKey) const
 {
 	if (UMovieSceneControlRigParameterSection* Section = Model->GetFKControlRigSection())
 	{
@@ -3324,8 +3339,10 @@ bool UAnimSequencerController::SetCurveControlKey(const FName& CurveName, const
 						ParameterCurvePair->ParameterCurve.ChangeFrameResolution(CurveRate, MaxChannelRate);
 						ParameterCurvePair->ParameterCurve.SetTickResolution(MaxChannelRate);
 					}					
-							
-					const FKeyHandle KeyHandle = ParameterCurvePair->ParameterCurve.GetData().UpdateOrAddKey(ParameterCurvePair->ParameterCurve.GetTickResolution().AsFrameNumber(Key.Time), MovieSceneValue);
+
+					const FFrameNumber FrameNumber = ParameterCurvePair->ParameterCurve.GetTickResolution().AsFrameTime(Key.Time).RoundToFrame();
+					check(bUpdateKey || ParameterCurvePair->ParameterCurve.GetData().FindKey(FrameNumber) == INDEX_NONE);
+					const FKeyHandle KeyHandle = ParameterCurvePair->ParameterCurve.GetData().UpdateOrAddKey(FrameNumber, MovieSceneValue);
 					return KeyHandle != FKeyHandle::Invalid();
 				}
 				else
diff --git a/Engine/Plugins/Animation/AnimationData/Source/AnimationData/Private/AnimSequencerController.h b/Engine/Plugins/Animation/AnimationData/Source/AnimationData/Private/AnimSequencerController.h
--- a/Engine/Plugins/Animation/AnimationData/Source/AnimationData/Private/AnimSequencerController.h
+++ b/Engine/Plugins/Animation/AnimationData/Source/AnimationData/Private/AnimSequencerController.h
@@ -111,7 +111,7 @@ private:
 	bool RenameCurveControl(const FName& CurveName, const FName& NewCurveName) const;
 	bool RemoveCurveControl(const FName& CurveName) const;
 	bool SetCurveControlKeys(const FName& CurveName, const TArray<FRichCurveKey>& CurveKeys) const;
-	bool SetCurveControlKey(const FName& CurveName, const FRichCurveKey& Key) const;
+	bool SetCurveControlKey(const FName& CurveName, const FRichCurveKey& Key, bool bUpdateKey) const;
 	bool RemoveCurveControlKey(const FName& CurveName, float Time) const;
 	bool DuplicateCurveControl(const FName& CurveName, const FName& DuplicateCurveName) const;
 
