Fix use-after free when duplicating animation attributes

diff --git a/Engine/Plugins/Animation/AnimationData/Source/AnimationData/Private/AnimSequencerController.cpp b/Engine/Plugins/Animation/AnimationData/Source/AnimationData/Private/AnimSequencerController.cpp
--- a/Engine/Plugins/Animation/AnimationData/Source/AnimationData/Private/AnimSequencerController.cpp
+++ b/Engine/Plugins/Animation/AnimationData/Source/AnimationData/Private/AnimSequencerController.cpp
@@ -2110,9 +2110,10 @@ bool UAnimSequencerController::DuplicateAttribute(const FAnimationAttributeIdent
 					IAnimationDataModel::FEvaluationAndModificationLock Lock(*ModelInterface);
 					FTransaction Transaction = ConditionalTransaction(LOCTEXT("DuplicateAttribute", "Duplicating Animation Attribute"), bShouldTransact);
 
-					FAnimatedBoneAttribute& DuplicateAttribute = Model->AnimatedBoneAttributes.AddDefaulted_GetRef();
+					FAnimatedBoneAttribute DuplicateAttribute;
 					DuplicateAttribute.Identifier = NewAttributeIdentifier;
 					DuplicateAttribute.Curve = AttributePtr->Curve;
+					Model->AnimatedBoneAttributes.Add(DuplicateAttribute);
 
 					FAttributeAddedPayload Payload;
 					Payload.Identifier = NewAttributeIdentifier;
