Fix crash from incorrect bone index type usage in UAnimationSequencerDataModel::GeneratePoseData

Also fixed bone index usage in FAnimPose::GetPose

diff --git a/Engine/Plugins/Animation/AnimationData/Source/AnimationData/Private/AnimSequencerDataModel.cpp b/Engine/Plugins/Animation/AnimationData/Source/AnimationData/Private/AnimSequencerDataModel.cpp
--- a/Engine/Plugins/Animation/AnimationData/Source/AnimationData/Private/AnimSequencerDataModel.cpp
+++ b/Engine/Plugins/Animation/AnimationData/Source/AnimationData/Private/AnimSequencerDataModel.cpp
@@ -1174,13 +1174,13 @@ void UAnimationSequencerDataModel::GeneratePoseData(UControlRig* ControlRig, FAn
 
 			UE::Anim::Retargeting::FRetargetingScope RetargetingScope(GetSkeleton(), RigPose, EvaluationContext);
 			
+			const FReferenceSkeleton& MeshRefSkeleton = RequiredBones.GetReferenceSkeleton();
+			// Called during compression that can occur while GC is in progress, marking weakptrs as unreachable temporarily
+			const FReferenceSkeleton& SkeletonRefSkeleton = RequiredBones.GetSkeletonAsset(true)->GetReferenceSkeleton();
+
 			// Populate bone/curve elements to Pose/Curve indices
 			{
 				QUICK_SCOPE_CYCLE_COUNTER(STAT_GetMappings);
-				const FReferenceSkeleton& MeshRefSkeleton = RequiredBones.GetReferenceSkeleton();
-				// Called during compression that can occur while GC is in progress, marking weakptrs as unreachable temporarily
-				const FReferenceSkeleton& SkeletonRefSkeleton = RequiredBones.GetSkeletonAsset(true)->GetReferenceSkeleton();
-
 				RigHierarchy->ForEach<FRigBoneElement>([&MeshRefSkeleton, &RequiredBones, &SkeletonRefSkeleton, &RigHierarchy, &RetargetingScope, &RigPose](const FRigBoneElement* BoneElement) -> bool
 				{
 					const FName& BoneName = BoneElement->GetName();
@@ -1217,7 +1217,7 @@ void UAnimationSequencerDataModel::GeneratePoseData(UControlRig* ControlRig, FAn
 			}
 
 			// Apply any additive transform curves - if requested and any are set
-			if (!RigPose.GetBoneContainer().ShouldUseSourceData())
+			if (!RequiredBones.ShouldUseSourceData())
 			{
 				for (const FTransformCurve& TransformCurve : GetTransformCurves())
 				{
@@ -1232,7 +1232,8 @@ void UAnimationSequencerDataModel::GeneratePoseData(UControlRig* ControlRig, FAn
 					// note we're not checking Curve.GetCurveTypeFlags() yet
 					FTransform Value = TransformCurve.Evaluate(static_cast<float>(EvaluationContext.SampleFrameRate.AsSeconds(EvaluationContext.SampleTime)), 1.f);
 
-					const FCompactPoseBoneIndex BoneIndex(RigPose.GetBoneContainer().GetPoseBoneIndexForBoneName(CurveName));
+					const int32 SkeletonBoneIndex = SkeletonRefSkeleton.FindBoneIndex(CurveName);
+					const FCompactPoseBoneIndex BoneIndex(RequiredBones.GetCompactPoseIndexFromSkeletonIndex(SkeletonBoneIndex));
 					if(BoneIndex != INDEX_NONE)
 					{
 						const FTransform LocalTransform = RigPose[BoneIndex];
diff --git a/Engine/Source/Editor/AnimationBlueprintLibrary/Private/AnimPose.cpp b/Engine/Source/Editor/AnimationBlueprintLibrary/Private/AnimPose.cpp
--- a/Engine/Source/Editor/AnimationBlueprintLibrary/Private/AnimPose.cpp
+++ b/Engine/Source/Editor/AnimationBlueprintLibrary/Private/AnimPose.cpp
@@ -92,7 +92,8 @@ void FAnimPose::GetPose(FCompactPose& InOutCompactPose) const
 		for (int32 Index = 0; Index < BoneNames.Num(); ++Index)
 		{
 			const FName& BoneName = BoneNames[Index];
-			const FCompactPoseBoneIndex PoseBoneIndex = FCompactPoseBoneIndex(InOutCompactPose.GetBoneContainer().GetPoseBoneIndexForBoneName(BoneName));
+			const int32 MeshBoneIndex = InOutCompactPose.GetBoneContainer().GetPoseBoneIndexForBoneName(BoneName);
+			const FCompactPoseBoneIndex PoseBoneIndex = InOutCompactPose.GetBoneContainer().MakeCompactPoseIndex(FMeshPoseBoneIndex(MeshBoneIndex));
 			if (PoseBoneIndex != INDEX_NONE)
 			{
 				InOutCompactPose[PoseBoneIndex] = LocalSpacePoses[Index];
