Fix ensure when transforming bone indices during animation compression

Use the non-deprecated conversion function

diff --git a/Engine/Plugins/Animation/AnimationData/Source/AnimationData/Private/AnimSequencerDataModel.cpp b/Engine/Plugins/Animation/AnimationData/Source/AnimationData/Private/AnimSequencerDataModel.cpp
--- a/Engine/Plugins/Animation/AnimationData/Source/AnimationData/Private/AnimSequencerDataModel.cpp
+++ b/Engine/Plugins/Animation/AnimationData/Source/AnimationData/Private/AnimSequencerDataModel.cpp
@@ -1232,8 +1232,8 @@ void UAnimationSequencerDataModel::GeneratePoseData(UControlRig* ControlRig, FAn
 					// note we're not checking Curve.GetCurveTypeFlags() yet
 					FTransform Value = TransformCurve.Evaluate(static_cast<float>(EvaluationContext.SampleFrameRate.AsSeconds(EvaluationContext.SampleTime)), 1.f);
 
-					const int32 SkeletonBoneIndex = SkeletonRefSkeleton.FindBoneIndex(CurveName);
-					const FCompactPoseBoneIndex BoneIndex(RequiredBones.GetCompactPoseIndexFromSkeletonIndex(SkeletonBoneIndex));
+					const FSkeletonPoseBoneIndex SkeletonBoneIndex = FSkeletonPoseBoneIndex(SkeletonRefSkeleton.FindBoneIndex(CurveName));
+					const FCompactPoseBoneIndex BoneIndex(RequiredBones.GetCompactPoseIndexFromSkeletonPoseIndex(SkeletonBoneIndex));
 					if(BoneIndex != INDEX_NONE)
 					{
 						const FTransform LocalTransform = RigPose[BoneIndex];
@@ -1273,7 +1273,7 @@ void UAnimationSequencerDataModel::GeneratePoseData(UControlRig* ControlRig, FAn
 				// Evaluate attributes at requested time interval
 				for (const FAnimatedBoneAttribute& Attribute : AnimatedBoneAttributes)
 				{
-					const FCompactPoseBoneIndex PoseBoneIndex = RequiredBones.GetCompactPoseIndexFromSkeletonIndex(Attribute.Identifier.GetBoneIndex());
+					const FCompactPoseBoneIndex PoseBoneIndex = RequiredBones.GetCompactPoseIndexFromSkeletonPoseIndex(FSkeletonPoseBoneIndex(Attribute.Identifier.GetBoneIndex()));
 					// Only add attribute if the bone its tied to exists in the currently evaluated set of bones
 					if(PoseBoneIndex.IsValid())
 					{
