Fixup incorrect by-value to by-const-ref

diff --git a/Engine/Plugins/Animation/AnimationData/Source/AnimationData/Private/AnimSequencerDataModel.cpp b/Engine/Plugins/Animation/AnimationData/Source/AnimationData/Private/AnimSequencerDataModel.cpp
--- a/Engine/Plugins/Animation/AnimationData/Source/AnimationData/Private/AnimSequencerDataModel.cpp
+++ b/Engine/Plugins/Animation/AnimationData/Source/AnimationData/Private/AnimSequencerDataModel.cpp
@@ -910,7 +910,7 @@ void UAnimationSequencerDataModel::GenerateLegacyCurveData()
 										}
 									}							
 
-									if (const FScalarParameterNameAndCurve* ScalarCurve = ScalarCurves.FindByPredicate([Element](FScalarParameterNameAndCurve Curve)
+									if (const FScalarParameterNameAndCurve* ScalarCurve = ScalarCurves.FindByPredicate([Element](const FScalarParameterNameAndCurve& Curve)
 									{
 										return Curve.ParameterName == Element->GetFName();
 									}))
