Always force-set float/transform controls during AnimSequencerData evaluation

diff --git a/Engine/Plugins/Animation/AnimationData/Source/AnimationData/Private/AnimSequencerDataModel.cpp b/Engine/Plugins/Animation/AnimationData/Source/AnimationData/Private/AnimSequencerDataModel.cpp
--- a/Engine/Plugins/Animation/AnimationData/Source/AnimationData/Private/AnimSequencerDataModel.cpp
+++ b/Engine/Plugins/Animation/AnimationData/Source/AnimationData/Private/AnimSequencerDataModel.cpp
@@ -1325,7 +1325,7 @@ void UAnimationSequencerDataModel::EvaluateTrack(UMovieSceneControlRigParameterT
 					FRigControlElement* ControlElement = ControlRig->FindControl(Name);
 					if (ControlElement && ControlElement->Settings.ControlType == ERigControlType::Float)
 					{
-						RigHierarchy->SetControlValue(ControlElement, FRigControlValue::Make<float>(Value), ERigControlValueType::Current, false, false, false, false);
+						RigHierarchy->SetControlValue(ControlElement, FRigControlValue::Make<float>(Value), ERigControlValueType::Current, false, true, false, false);
 					}
 				}
 			}
@@ -1461,7 +1461,7 @@ void UAnimationSequencerDataModel::EvaluateTrack(UMovieSceneControlRigParameterT
 						
 							EulerTransform = FEulerTransform(FinalTransform);
 						}
-						RigHierarchy->SetControlValue(ControlElement, FRigControlValue::Make<FRigControlValue::FEulerTransform_Float>(EulerTransform), ERigControlValueType::Current, false, false, false, false);
+						RigHierarchy->SetControlValue(ControlElement, FRigControlValue::Make<FRigControlValue::FEulerTransform_Float>(EulerTransform), ERigControlValueType::Current, false, true, false, false);
 					}
 				}
 			}
