Add path name to assertions

diff --git a/Engine/Plugins/Animation/AnimationData/Source/AnimationData/Private/AnimSequencerDataModel.cpp b/Engine/Plugins/Animation/AnimationData/Source/AnimationData/Private/AnimSequencerDataModel.cpp
--- a/Engine/Plugins/Animation/AnimationData/Source/AnimationData/Private/AnimSequencerDataModel.cpp
+++ b/Engine/Plugins/Animation/AnimationData/Source/AnimationData/Private/AnimSequencerDataModel.cpp
@@ -944,19 +944,19 @@ void UAnimationSequencerDataModel::ValidateData() const
 
 void UAnimationSequencerDataModel::ValidateSequencerData() const
 {
-	checkf(MovieScene, TEXT("No Movie Scene found for SequencerDataModel"));
+	checkf(MovieScene, TEXT("%s: No Movie Scene found for SequencerDataModel"), *GetPathName());
 
 	const int32 NumberOfTracks = MovieScene->GetTracks().Num();
-	checkf(NumberOfTracks == 1, TEXT("Invalid number of Tracks in Movie Scene expected 1 but found %i"), NumberOfTracks);
+	checkf(NumberOfTracks == 1, TEXT("%s: Invalid number of Tracks in Movie Scene expected 1 but found %i"), *GetPathName(), NumberOfTracks);
 		
 	const UMovieSceneControlRigParameterTrack* Track = MovieScene->FindTrack<UMovieSceneControlRigParameterTrack>();
-	checkf(Track, TEXT("Unable to find Control Rig Track"));
+	checkf(Track, TEXT("%s: Unable to find Control Rig Track"), *GetPathName());
 
 	const int32 NumberOfSections = Track->GetAllSections().Num();
-	checkf(NumberOfSections == 1, TEXT("Invalid number of Sections found for Control Rig Track expected 1 but found %i"), NumberOfSections);
+	checkf(NumberOfSections == 1, TEXT("%s: Invalid number of Sections found for Control Rig Track expected 1 but found %i"), *GetPathName(), NumberOfSections);
 
 	const UMovieSceneControlRigParameterSection* Section = GetFKControlRigSection();
-	checkf(Section, TEXT("Unable to find Control Rig Section"));
+	checkf(Section, TEXT("%s: Unable to find Control Rig Section"), *GetPathName());
 }
 
 void UAnimationSequencerDataModel::ValidateControlRigData() const
