Ensure that type hash for non-memory-trivial animation attributes is set-able/used

diff --git a/Engine/Plugins/Animation/AnimationData/Source/AnimationData/Private/AnimSequencerDataModel.cpp b/Engine/Plugins/Animation/AnimationData/Source/AnimationData/Private/AnimSequencerDataModel.cpp
--- a/Engine/Plugins/Animation/AnimationData/Source/AnimationData/Private/AnimSequencerDataModel.cpp
+++ b/Engine/Plugins/Animation/AnimationData/Source/AnimationData/Private/AnimSequencerDataModel.cpp
@@ -651,18 +651,28 @@ FGuid UAnimationSequencerDataModel::GenerateGuid() const
 		{
 			Sha.Update(reinterpret_cast<const uint8*>(&Data), sizeof(Data));
 		};
-		
+
 		for (const FAnimatedBoneAttribute& Attribute : AnimatedBoneAttributes)
 		{
 			UpdateSHAWithArray(Attribute.Identifier.GetName().ToString().GetCharArray());
 			UpdateSHAWithArray(Attribute.Identifier.GetBoneName().ToString().GetCharArray());
 			UpdateWithData(Attribute.Identifier.GetBoneIndex());
 			UpdateSHAWithArray(Attribute.Identifier.GetType()->GetFName().ToString().GetCharArray());
-			const uint32 StructSize = Attribute.Identifier.GetType()->GetPropertiesSize();
+			const UScriptStruct* TypeStruct = Attribute.Identifier.GetType();
+			const uint32 StructSize = TypeStruct->GetPropertiesSize();
+			const bool bHasTypeHash = TypeStruct->GetCppStructOps()->HasGetTypeHash();
 			for (const FAttributeKey& Key : Attribute.Curve.GetConstRefOfKeys())
 			{
 				UpdateWithData(Key.Time);
-				Sha.Update(Key.GetValuePtr<uint8>(), StructSize);
+				if (bHasTypeHash)
+				{
+					const uint32 KeyHash = TypeStruct->GetStructTypeHash(Key.GetValuePtr<uint8>());
+					UpdateWithData(KeyHash);
+				}
+				else
+				{
+					Sha.Update(Key.GetValuePtr<uint8>(), StructSize);
+				}
 			}
 		}
 
diff --git a/Engine/Source/Runtime/Engine/Classes/Animation/BuiltInAttributeTypes.h b/Engine/Source/Runtime/Engine/Classes/Animation/BuiltInAttributeTypes.h
--- a/Engine/Source/Runtime/Engine/Classes/Animation/BuiltInAttributeTypes.h
+++ b/Engine/Source/Runtime/Engine/Classes/Animation/BuiltInAttributeTypes.h
@@ -87,6 +87,11 @@ struct FStringAnimationAttribute
 	FString Value;
 };
 
+inline uint32 GetTypeHash(const FStringAnimationAttribute& Key)
+{
+	return GetTypeHash(Key.Value);
+}
+
 /** Attribute type supporting the legacy TVariant<FTransform> attributes */
 USTRUCT(BlueprintType)
 struct FTransformAnimationAttribute
