Reduce cook-time memory usage for AnimSequencerDataModel
    - Lazy initialize RigHierarchy
    - Reset RigHierarchy on WillNeverCacheCookedPlatformDataAgain

diff --git a/Engine/Plugins/Animation/AnimationData/Source/AnimationData/Private/AnimSequencerController.cpp b/Engine/Plugins/Animation/AnimationData/Source/AnimationData/Private/AnimSequencerController.cpp
--- a/Engine/Plugins/Animation/AnimationData/Source/AnimationData/Private/AnimSequencerController.cpp
+++ b/Engine/Plugins/Animation/AnimationData/Source/AnimationData/Private/AnimSequencerController.cpp
@@ -1309,7 +1309,7 @@ bool UAnimSequencerController::SetCurveAttributes(const FAnimationCurveIdentifie
 
 			if (UFKControlRig* FKRig = Cast<UFKControlRig>(ControlRig))
 			{
-				const URigHierarchy* Hierarchy = FKRig->GetHierarchy();
+				const URigHierarchy* Hierarchy = Model->GetControlRigHierarchy();
 				if (Hierarchy ||  IgnoreSkeletonValidation())
 				{
 					const FRigElementKey CurveKey(URigHierarchy::GetSanitizedName(CurveId.CurveName), ERigElementType::Curve);
@@ -1545,7 +1545,7 @@ bool UAnimSequencerController::SetBoneTrackKeys(FName BoneName, const TArray<FVe
 			{
 				if (UControlRig* ControlRig = Model->GetControlRig())
 				{
-					if (URigHierarchy* Hierarchy = ControlRig->GetHierarchy())
+					if (URigHierarchy* Hierarchy = Model->GetControlRigHierarchy())
 					{
 						const FRigElementKey BoneKey(URigHierarchy::GetSanitizedName(BoneName), ERigElementType::Bone);
 						const FRigElementKey BoneNameControlKey(UFKControlRig::GetControlName(BoneKey.Name, ERigElementType::Bone), ERigElementType::Control);
@@ -2152,7 +2152,7 @@ void UAnimSequencerController::RemoveUnusedControlsAndCurves() const
 	{
 		if (UControlRig* ControlRig = Model->GetControlRig())
 		{
-			if (URigHierarchy* Hierarchy = ControlRig->GetHierarchy())
+			if (URigHierarchy* Hierarchy = Model->GetControlRigHierarchy())
 			{
 				TArray<FRigElementKey> ElementsToRemove;
 				
@@ -2207,7 +2207,7 @@ void UAnimSequencerController::UpdateWithSkeleton(USkeleton* TargetSkeleton, boo
 		OpenBracket(LOCTEXT("SettingNewskeleton", "Updating Skeleton for Animation Data Model"), bShouldTransact);
 		{
 			// (re-)generate the rig hierarchy
-			Model->InitializeFKControlRig(CastChecked<UFKControlRig>(Model->GetControlRig()), TargetSkeleton);
+			Model->InitializeFKControlRig(CastChecked<UFKControlRig>(Model->GetControlRig()), TargetSkeleton, true);
 
 			// Remove any bone tracks that do not exist in the new hierarchy
 			RemoveBoneTracksMissingFromSkeleton(TargetSkeleton, bShouldTransact);
@@ -2382,8 +2382,8 @@ void UAnimSequencerController::InitializeModel()
 				{
 					if(UFKControlRig* ControlRig = NewObject<UFKControlRig>(Track, UFKControlRig::StaticClass(), NAME_None, RF_Transactional))
 					{
-						Model->InitializeFKControlRig(ControlRig, Model->GetSkeleton());
-
+						Model->InitializeFKControlRig(ControlRig, Model->GetSkeleton(), true);
+						
 						// Remove all control elements (start fresh)
 						if (URigHierarchy* Hierarchy = ControlRig->GetHierarchy())
 						{
@@ -2417,7 +2417,7 @@ void UAnimSequencerController::InitializeModel()
 						else
 						{
 							ReportError(LOCTEXT("InvalidRigHierarchy", "Unable to retrieve valid URigHierarchy"));
-						}
+						}					
 
 						UMovieSceneControlRigParameterSection* Section = Cast<UMovieSceneControlRigParameterSection>(Track->CreateControlRigSection(0, ControlRig, true));				
 						Section->SetRange(DataRange);
@@ -2450,7 +2450,7 @@ bool UAnimSequencerController::AddBoneControl(const FName& BoneName) const
 	{
 		if (UFKControlRig* FKRig = Cast<UFKControlRig>(Section->GetControlRig()))
 		{
-			if(URigHierarchy* Hierarchy = FKRig->GetHierarchy())
+			if(URigHierarchy* Hierarchy = Model->GetControlRigHierarchy())
 			{
 				const FRigElementKey BoneKey(URigHierarchy::GetSanitizedName(BoneName), ERigElementType::Bone);
 				const FRigElementKey BoneNameControlKey(UFKControlRig::GetControlName(BoneKey.Name, ERigElementType::Bone), ERigElementType::Control);
@@ -2555,7 +2555,7 @@ bool UAnimSequencerController::RemoveBoneControl(const FName& BoneName) const
 	{
 		if (UFKControlRig* FKRig = Cast<UFKControlRig>(Section->GetControlRig()))
 		{
-			if(URigHierarchy* Hierarchy = FKRig->GetHierarchy())
+			if(URigHierarchy* Hierarchy = Model->GetControlRigHierarchy())
 			{
 				const FRigElementKey BoneNameControlKey(UFKControlRig::GetControlName(URigHierarchy::GetSanitizedName(BoneName), ERigElementType::Bone), ERigElementType::Control);
 				
@@ -2612,7 +2612,7 @@ bool UAnimSequencerController::SetBoneCurveKeys(const FName& BoneName, const TAr
 		UControlRig* ControlRig = Section->GetControlRig();
 		if (UFKControlRig* FKRig = Cast<UFKControlRig>(ControlRig))
 		{
-			const URigHierarchy* Hierarchy = FKRig->GetHierarchy();
+			const URigHierarchy* Hierarchy = Model->GetControlRigHierarchy();
 			if(Hierarchy || IgnoreSkeletonValidation())
 			{
 				const FRigElementKey BoneKey(URigHierarchy::GetSanitizedName(BoneName), ERigElementType::Bone);
@@ -2804,7 +2804,7 @@ bool UAnimSequencerController::UpdateBoneCurveKeys(const FName& BoneName, const
 		UControlRig* ControlRig = Section->GetControlRig();
 		if (UFKControlRig* FKRig = Cast<UFKControlRig>(ControlRig))
 		{
-			if(URigHierarchy* Hierarchy = FKRig->GetHierarchy())
+			if(URigHierarchy* Hierarchy = Model->GetControlRigHierarchy())
 			{
 				const FRigElementKey BoneKey(URigHierarchy::GetSanitizedName(BoneName), ERigElementType::Bone);
 				const FRigElementKey BoneNameControlKey(UFKControlRig::GetControlName(BoneKey.Name, ERigElementType::Bone), ERigElementType::Control);
@@ -2900,7 +2900,7 @@ bool UAnimSequencerController::RemoveBoneCurveKey(const FName& BoneName, float T
 	
 			if (UFKControlRig* FKRig = Cast<UFKControlRig>(ControlRig))
 			{
-				if(const URigHierarchy* Hierarchy = FKRig->GetHierarchy())
+				if(const URigHierarchy* Hierarchy = Model->GetControlRigHierarchy())
 				{
 					const FRigElementKey BoneKey(URigHierarchy::GetSanitizedName(BoneName), ERigElementType::Bone);
 					const FRigElementKey BoneNameControlKey(UFKControlRig::GetControlName(BoneKey.Name, ERigElementType::Bone), ERigElementType::Control);
@@ -2973,7 +2973,7 @@ bool UAnimSequencerController::AddCurveControl(const FName& CurveName) const
 		UControlRig* ControlRig = Section->GetControlRig();
 		if (UFKControlRig* FKRig = Cast<UFKControlRig>(ControlRig))
 		{
-			if(URigHierarchy* Hierarchy = FKRig->GetHierarchy())
+			if(URigHierarchy* Hierarchy = Model->GetControlRigHierarchy())
 			{				
 				const FRigElementKey CurveKey(URigHierarchy::GetSanitizedName(CurveName), ERigElementType::Curve);
 				const FRigElementKey CurveControlKey(UFKControlRig::GetControlName(CurveKey.Name, ERigElementType::Curve), ERigElementType::Control);
@@ -3055,7 +3055,7 @@ bool UAnimSequencerController::RenameCurveControl(const FName& CurveName, const
 	{
 		if (UControlRig* ControlRig = Section->GetControlRig())
 		{
-			if (URigHierarchy* RigHierarchy = ControlRig->GetHierarchy())
+			if (URigHierarchy* RigHierarchy = Model->GetControlRigHierarchy())
 			{
 				if (URigHierarchyController* Controller = RigHierarchy->GetController())
 				{
@@ -3144,7 +3144,7 @@ bool UAnimSequencerController::RemoveCurveControl(const FName& CurveName) const
 
 		if (UFKControlRig* FKRig = Cast<UFKControlRig>(ControlRig))
 		{
-			if(URigHierarchy* Hierarchy = FKRig->GetHierarchy())
+			if(URigHierarchy* Hierarchy = Model->GetControlRigHierarchy())
 			{									
 				const FRigElementKey CurveKey(URigHierarchy::GetSanitizedName(CurveName), ERigElementType::Curve);
 				const FRigElementKey CurveControlKey(UFKControlRig::GetControlName(CurveKey.Name, ERigElementType::Curve), ERigElementType::Control);
@@ -3208,7 +3208,7 @@ bool UAnimSequencerController::SetCurveControlKeys(const FName& CurveName, const
 
 		if (UFKControlRig* FKRig = Cast<UFKControlRig>(ControlRig))
 		{
-			const URigHierarchy* Hierarchy = FKRig->GetHierarchy();
+			const URigHierarchy* Hierarchy = Model->GetControlRigHierarchy();
 
 			if (Hierarchy ||  IgnoreSkeletonValidation())
 			{
@@ -3296,7 +3296,7 @@ bool UAnimSequencerController::SetCurveControlKey(const FName& CurveName, const
 
 		if (UFKControlRig* FKRig = Cast<UFKControlRig>(ControlRig))
 		{
-			if(const URigHierarchy* Hierarchy = FKRig->GetHierarchy())
+			if(const URigHierarchy* Hierarchy = Model->GetControlRigHierarchy())
 			{
 				const FRigElementKey CurveKey(URigHierarchy::GetSanitizedName(CurveName), ERigElementType::Curve);
 				const FRigElementKey CurveControlKey(UFKControlRig::GetControlName(CurveKey.Name, ERigElementType::Curve), ERigElementType::Control);
@@ -3385,7 +3385,7 @@ bool UAnimSequencerController::RemoveCurveControlKey(const FName& CurveName, flo
 
 		if (UFKControlRig* FKRig = Cast<UFKControlRig>(ControlRig))
 		{
-			if(const URigHierarchy* Hierarchy = FKRig->GetHierarchy())
+			if(const URigHierarchy* Hierarchy = Model->GetControlRigHierarchy())
 			{
 				const FFrameNumber FrameNumber = Model->MovieScene->GetTickResolution().AsFrameNumber(Time);
 				const FRigElementKey CurveKey(URigHierarchy::GetSanitizedName(CurveName), ERigElementType::Curve);
@@ -3449,7 +3449,7 @@ bool UAnimSequencerController::DuplicateCurveControl(const FName& CurveName, con
 	{
 		if (UFKControlRig* FKRig = Cast<UFKControlRig>(Section->GetControlRig()))
 		{
-			if (URigHierarchy* Hierarchy = FKRig->GetHierarchy())
+			if (URigHierarchy* Hierarchy = Model->GetControlRigHierarchy())
 			{
 				if (URigHierarchyController* HierarchyController = Hierarchy->GetController())
 				{
diff --git a/Engine/Plugins/Animation/AnimationData/Source/AnimationData/Private/AnimSequencerDataModel.cpp b/Engine/Plugins/Animation/AnimationData/Source/AnimationData/Private/AnimSequencerDataModel.cpp
--- a/Engine/Plugins/Animation/AnimationData/Source/AnimationData/Private/AnimSequencerDataModel.cpp
+++ b/Engine/Plugins/Animation/AnimationData/Source/AnimationData/Private/AnimSequencerDataModel.cpp
@@ -107,7 +107,36 @@ USkeleton* UAnimationSequencerDataModel::GetSkeleton() const
 	return Skeleton;
 }
 
-void UAnimationSequencerDataModel::InitializeFKControlRig(UFKControlRig* FKControlRig, USkeleton* Skeleton) const
+void UAnimationSequencerDataModel::InitializeRigHierarchy(UFKControlRig* FKControlRig, USkeleton* Skeleton) const
+{
+	FScopeLock Lock(&EvaluationLock);
+	
+	UFKControlRig::FRigElementInitializationOptions InitOptions;
+	InitOptions.bImportCurves = false;	
+	if(UMovieSceneControlRigParameterSection* Section = GetFKControlRigSection())
+	{
+		for (const FScalarParameterNameAndCurve& AnimCurve : Section->GetScalarParameterNamesAndCurves())
+		{
+			InitOptions.CurveNames.Add(UFKControlRig::GetControlTargetName(AnimCurve.ParameterName, ERigElementType::Curve));
+		}
+
+		for (const FTransformParameterNameAndCurves& BoneCurve : Section->GetTransformParameterNamesAndCurves())
+		{
+			InitOptions.BoneNames.Add(UFKControlRig::GetControlTargetName(BoneCurve.ParameterName, ERigElementType::Bone));
+		}
+	}
+	InitOptions.bGenerateBoneControls = InitOptions.BoneNames.Num() > 0;
+	FKControlRig->SetInitializationOptions(InitOptions);
+
+	FKControlRig->Initialize();
+
+	FKControlRig->SetApplyMode(UseDirectFKControlRigMode == 1 ? EControlRigFKRigExecuteMode::Direct : EControlRigFKRigExecuteMode::Replace);
+	FKControlRig->SetBoneInitialTransformsFromRefSkeleton(Skeleton->GetReferenceSkeleton());
+	FKControlRig->Evaluate_AnyThread();
+	bRigHierarchyInitialized = true;
+}
+
+void UAnimationSequencerDataModel::InitializeFKControlRig(UFKControlRig* FKControlRig, USkeleton* Skeleton, bool bForceHierarchyInitialization /*=false*/) const
 {
 	checkf(FKControlRig, TEXT("Invalid FKControlRig provided"));
 	if (Skeleton)
@@ -117,35 +146,17 @@ void UAnimationSequencerDataModel::InitializeFKControlRig(UFKControlRig* FKContr
 		FKControlRig->SetObjectBinding(MakeShared<FControlRigObjectBinding>());
 		FKControlRig->GetObjectBinding()->BindToObject(Skeleton);
 
-		UFKControlRig::FRigElementInitializationOptions InitOptions;
-		InitOptions.bImportCurves = false;	
-		if(UMovieSceneControlRigParameterSection* Section = GetFKControlRigSection())
+		if (!IsRunningCookCommandlet() || bForceHierarchyInitialization)
 		{
-			for (const FScalarParameterNameAndCurve& AnimCurve : Section->GetScalarParameterNamesAndCurves())
-			{
-				InitOptions.CurveNames.Add(UFKControlRig::GetControlTargetName(AnimCurve.ParameterName, ERigElementType::Curve));
-			}
-
-			for (const FTransformParameterNameAndCurves& BoneCurve : Section->GetTransformParameterNamesAndCurves())
-			{
-				InitOptions.BoneNames.Add(UFKControlRig::GetControlTargetName(BoneCurve.ParameterName, ERigElementType::Bone));
-			}
+			InitializeRigHierarchy(FKControlRig, Skeleton);
 		}
-		InitOptions.bGenerateBoneControls = InitOptions.BoneNames.Num() > 0;
-		FKControlRig->SetInitializationOptions(InitOptions);
-
-		FKControlRig->Initialize();
-
-		FKControlRig->SetApplyMode(UseDirectFKControlRigMode == 1 ? EControlRigFKRigExecuteMode::Direct : EControlRigFKRigExecuteMode::Replace);
-		FKControlRig->SetBoneInitialTransformsFromRefSkeleton(Skeleton->GetReferenceSkeleton());
-		FKControlRig->Evaluate_AnyThread();
 		
 		UnlockEvaluationAndModification();
 	}
 }
 
 UControlRig* UAnimationSequencerDataModel::GetControlRig() const
-{
+{	
 	if(const UMovieSceneControlRigParameterTrack* Track = GetControlRigTrack())
 	{
 		return Track->GetControlRig();
@@ -179,7 +190,7 @@ void UAnimationSequencerDataModel::PostLoad()
 		if (const UMovieSceneControlRigParameterSection* Section = GetFKControlRigSection())
 		{
 			if (UFKControlRig* ControlRig = Cast<UFKControlRig>(Section->GetControlRig()))
-			{
+			{				
 				InitializeFKControlRig(ControlRig, GetSkeleton());
 			}
 		}
@@ -223,6 +234,7 @@ void UAnimationSequencerDataModel::PreSave(FObjectPreSaveContext ObjectSaveConte
 void UAnimationSequencerDataModel::WillNeverCacheCookedPlatformDataAgain()
 {
 	Super::WillNeverCacheCookedPlatformDataAgain();
+	ClearControlRigData();
 }
 
 void UAnimationSequencerDataModel::PreEditUndo()
@@ -747,6 +759,15 @@ void UAnimationSequencerDataModel::Evaluate(FAnimationPoseData& InOutPoseData, c
 	if (UMovieSceneControlRigParameterTrack* Track = GetControlRigTrack())
 	{
 		FScopeLock Lock(&EvaluationLock);
+
+		if(!bRigHierarchyInitialized)
+		{
+			if (UFKControlRig* FKControlRig = Cast<UFKControlRig>(GetControlRig()))
+			{
+				InitializeRigHierarchy(FKControlRig, GetSkeleton());
+			}
+		}
+		
 		// Evaluates and applies control curves from track to ControlRig
 		EvaluateTrack(Track, EvaluationContext);
 
@@ -863,6 +884,27 @@ UMovieSceneControlRigParameterSection* UAnimationSequencerDataModel::GetFKContro
 	return nullptr;
 }
 
+URigHierarchy* UAnimationSequencerDataModel::GetControlRigHierarchy() const
+{	
+	if (const UMovieSceneControlRigParameterSection* Section = GetFKControlRigSection())
+	{
+		UControlRig* ControlRig = Section->GetControlRig();
+		if (UFKControlRig* FKRig = Cast<UFKControlRig>(ControlRig))
+		{
+			if (!bRigHierarchyInitialized)
+			{
+				InitializeRigHierarchy(FKRig, GetSkeleton());
+			}
+
+			return FKRig->GetHierarchy();
+		}
+
+		IAnimationDataController::ReportObjectErrorf(this, LOCTEXT("UnableToFindRigHierarchy", "Unable to retrieve RigHierarchy for ControlRig ({0})"), FText::FromString(ControlRig->GetPathName()));
+	}
+
+	return nullptr;
+}
+
 void UAnimationSequencerDataModel::GenerateLegacyCurveData()
 {
 	QUICK_SCOPE_CYCLE_COUNTER(STAT_GenerateLegacyCurveData);
@@ -1154,6 +1196,19 @@ void UAnimationSequencerDataModel::GenerateTransformKeysForControl(const FName&
 	}, &FrameNumbers);
 }
 
+void UAnimationSequencerDataModel::ClearControlRigData()
+{
+	if (bRigHierarchyInitialized)
+	{
+		FScopeLock Lock(&EvaluationLock);
+		if (UFKControlRig* FKControlRig = Cast<UFKControlRig>(GetControlRig()))
+		{
+			FKControlRig->GetHierarchy()->Reset();
+			bRigHierarchyInitialized = false;
+		}
+	}	
+}
+
 UMovieScene* UAnimationSequencerDataModel::GetMovieScene() const
 {
 	return MovieScene;
diff --git a/Engine/Plugins/Animation/AnimationData/Source/AnimationData/Private/AnimSequencerDataModel.h b/Engine/Plugins/Animation/AnimationData/Source/AnimationData/Private/AnimSequencerDataModel.h
--- a/Engine/Plugins/Animation/AnimationData/Source/AnimationData/Private/AnimSequencerDataModel.h
+++ b/Engine/Plugins/Animation/AnimationData/Source/AnimationData/Private/AnimSequencerDataModel.h
@@ -13,6 +13,7 @@ class UControlRig;
 class UMovieSceneControlRigParameterTrack;
 class UMovieSceneControlRigParameterSection;
 class UFKControlRig;
+class URigHierarchy;
 
 USTRUCT()
 struct FAnimationCurveMetaData
@@ -151,8 +152,10 @@ protected:
 
 	UMovieSceneControlRigParameterTrack* GetControlRigTrack() const;
 	UMovieSceneControlRigParameterSection* GetFKControlRigSection() const;
+	URigHierarchy* GetControlRigHierarchy() const;
 	USkeleton* GetSkeleton() const;
-	void InitializeFKControlRig(UFKControlRig* FKControlRig, USkeleton* Skeleton) const;
+	void InitializeRigHierarchy(UFKControlRig* FKControlRig, USkeleton* Skeleton) const;
+	void InitializeFKControlRig(UFKControlRig* FKControlRig, USkeleton* Skeleton, bool bForceHierarchyInitialization=false) const;
 	UControlRig* GetControlRig() const;
 	
 	void IterateTransformControlCurve(const FName& BoneName, TFunction<void(const FTransform&, const FFrameNumber&)> IterationFunction, const TArray<FFrameNumber>* InFrameNumbers = nullptr) const;
@@ -160,6 +163,7 @@ protected:
 	void RemoveOutOfDateControls() const;
 	
 	void GenerateTransformKeysForControl(const FName& BoneName, const TArray<FFrameNumber>& FrameNumbers, TArray<FTransform>& InOutTransforms) const;
+	void ClearControlRigData();
 private:	
 	/** Dynamic delegate event allows scripting to register to any broadcast-ed notify. */
 	UPROPERTY(BlueprintAssignable, Transient, Category = AnimationDataModel, meta = (ScriptName = "ModifiedEvent", AllowPrivateAccess = "true"))
@@ -189,6 +193,9 @@ private:
 	UPROPERTY(VisibleAnywhere, Category=AnimSequencer)
 	FGuid CachedRawDataGUID;
 
+	UPROPERTY(VisibleAnywhere, Transient, Category=AnimSequencer)
+	mutable bool bRigHierarchyInitialized = false;
+
 	// Scope lock to prevent contention around ControlRig->Evaluation()
 	mutable FCriticalSection EvaluationLock;
 
