Animation sequence stops and warning showing after undoing adding or deleting frames
#fix ensure that the MovieScene is not transacted in addition to the Controller Actions when changing the playback range

diff --git a/Engine/Plugins/Animation/AnimationData/Source/AnimationData/Private/AnimSequencerController.cpp b/Engine/Plugins/Animation/AnimationData/Source/AnimationData/Private/AnimSequencerController.cpp
--- a/Engine/Plugins/Animation/AnimationData/Source/AnimationData/Private/AnimSequencerController.cpp
+++ b/Engine/Plugins/Animation/AnimationData/Source/AnimationData/Private/AnimSequencerController.cpp
@@ -226,7 +226,7 @@ void UAnimSequencerController::SetMovieSceneRange(FFrameNumber InFrameNumber) co
 	const FFrameTime TickResolutionFrameNumber = FFrameRate::TransformTime(FFrameTime(InFrameNumber), Model->GetFrameRate(), TickRate);
 	
 	const TRange<FFrameNumber> DataRange = TRange<FFrameNumber>::Inclusive(FFrameNumber(0), TickResolutionFrameNumber.GetFrame());
-	MovieScene->SetPlaybackRange(DataRange);
+	MovieScene->SetPlaybackRange(DataRange, false);
 	MovieScene->SetPlaybackRangeLocked(false);
 	
 	FMovieSceneEditorData& EditorData = MovieScene->GetEditorData();
@@ -2368,7 +2368,7 @@ void UAnimSequencerController::InitializeModel()
 				MovieScene->SetDisplayRate(UAnimationSettings::Get()->GetDefaultFrameRate());
 
 				const TRange<FFrameNumber> DataRange = TRange<FFrameNumber>::Inclusive(FFrameNumber(0), 1);
-				MovieScene->SetPlaybackRange(DataRange);
+				MovieScene->SetPlaybackRange(DataRange, false);
 
 				if (AnimSequence->GetLinkerCustomVersion(FUE5MainStreamObjectVersion::GUID) < FUE5MainStreamObjectVersion::IntroducingAnimationDataModel)
 				{
