Clear default value when calling UpdateBoneTrackKeys in MovieSceneChannel

diff --git a/Engine/Plugins/Animation/AnimationData/Source/AnimationData/Private/AnimSequencerController.cpp b/Engine/Plugins/Animation/AnimationData/Source/AnimationData/Private/AnimSequencerController.cpp
--- a/Engine/Plugins/Animation/AnimationData/Source/AnimationData/Private/AnimSequencerController.cpp
+++ b/Engine/Plugins/Animation/AnimationData/Source/AnimationData/Private/AnimSequencerController.cpp
@@ -2840,12 +2840,15 @@ bool UAnimSequencerController::UpdateBoneCurveKeys(const FName& BoneName, const
 						{
 							// Translation
 							AddKeyToChannel(&ParameterCurvePair->Translation[ChannelIndex], FrameNumber, static_cast<float>(Location[ChannelIndex]), EMovieSceneKeyInterpolation::Linear);
+							ParameterCurvePair->Translation[ChannelIndex].ClearDefault();
 						
 							// Rotation
 							AddKeyToChannel(&ParameterCurvePair->Rotation[ChannelIndex], FrameNumber, static_cast<float>(EulerAngles[ChannelIndex]), EMovieSceneKeyInterpolation::Linear);
+							ParameterCurvePair->Rotation[ChannelIndex].ClearDefault();
 							
 							// Scaling							
-							AddKeyToChannel(&ParameterCurvePair->Scale[ChannelIndex], FrameNumber, static_cast<float>(Scale[ChannelIndex]), EMovieSceneKeyInterpolation::Linear);						
+							AddKeyToChannel(&ParameterCurvePair->Scale[ChannelIndex], FrameNumber, static_cast<float>(Scale[ChannelIndex]), EMovieSceneKeyInterpolation::Linear);
+							ParameterCurvePair->Scale[ChannelIndex].ClearDefault();
 						}
 					}
 					
@@ -3528,5 +3531,3 @@ bool UAnimSequencerController::DuplicateCurveControl(const FName& CurveName, con
 }
 
 #undef LOCTEXT_NAMESPACE // "AnimDataController"
-
-
