Remove unused code

diff --git a/Engine/Plugins/Animation/AnimationData/Source/AnimationData/Private/AnimSequencerController.cpp b/Engine/Plugins/Animation/AnimationData/Source/AnimationData/Private/AnimSequencerController.cpp
--- a/Engine/Plugins/Animation/AnimationData/Source/AnimationData/Private/AnimSequencerController.cpp
+++ b/Engine/Plugins/Animation/AnimationData/Source/AnimationData/Private/AnimSequencerController.cpp
@@ -340,7 +340,6 @@ bool UAnimSequencerController::RemoveBoneTracksMissingFromSkeleton(const USkelet
 		if (BoundControlRigObject == Skeleton)
 		{
 			TArray<FName> TracksToBeRemoved;
-			TArray<FName> TracksUpdated;
 			const FReferenceSkeleton& ReferenceSkeleton = Skeleton->GetReferenceSkeleton();
 
 			IAnimationDataModel::FEvaluationAndModificationLock Lock(*ModelInterface);
@@ -359,23 +358,16 @@ bool UAnimSequencerController::RemoveBoneTracksMissingFromSkeleton(const USkelet
 				}
 			}
 
-			if (TracksToBeRemoved.Num() || TracksUpdated.Num())
+			if (TracksToBeRemoved.Num())
 			{
 				FBracket Bracket = ConditionalBracket(LOCTEXT("RemoveBoneTracksMissingFromSkeleton", "Validating Bone Animation Track Data against Skeleton"), bShouldTransact);
 				for (const FName& TrackName : TracksToBeRemoved)
 				{
 					RemoveBoneTrack(TrackName);
 				}
-
-				for (const FName& TrackName : TracksUpdated)
-				{
-					FAnimationTrackChangedPayload Payload;
-					Payload.Name = TrackName;
-					Model->GetNotifier().Notify(EAnimDataModelNotifyType::TrackChanged, Payload);
-				}
 			}
 
-			return TracksToBeRemoved.Num() > 0 || TracksUpdated.Num() > 0;
+			return TracksToBeRemoved.Num() > 0;
 		}
 		else
 		{
