Added the ability to add comments to curves on animations

Comments are added from the curve's outliner menu. Editing can just be done by clicking the comment after it has been created
Also exposed the setting to view curves in a tree view in the timeline itself

NOTE: this removes a 4.5 core redirect from CurveName to LastObservedName as CurveName was reinstated as the canonical name for the curve, and having the redirect in place causes issues with tagged property serialization when new properties are added.

diff --git a/Engine/Plugins/Animation/AnimationData/Source/AnimationData/Private/AnimSequencerController.cpp b/Engine/Plugins/Animation/AnimationData/Source/AnimationData/Private/AnimSequencerController.cpp
--- a/Engine/Plugins/Animation/AnimationData/Source/AnimationData/Private/AnimSequencerController.cpp
+++ b/Engine/Plugins/Animation/AnimationData/Source/AnimationData/Private/AnimSequencerController.cpp
@@ -1085,6 +1085,49 @@ bool UAnimSequencerController::SetCurveColor(const FAnimationCurveIdentifier& Cu
 	return false;
 }
 
+bool UAnimSequencerController::SetCurveComment(const FAnimationCurveIdentifier& CurveId, const FString& Comment, bool bShouldTransact)
+{
+	ValidateModel();
+
+	if (CurveId.IsValid())
+	{
+		if (CurveId.CurveType == ERawCurveTrackTypes::RCT_Float)
+		{
+			if (const FFloatCurve* Curve = Model->FindMutableFloatCurveById(CurveId))
+			{
+				IAnimationDataModel::FEvaluationAndModificationLock Lock(*ModelInterface);
+				FTransaction Transaction = ConditionalTransaction(LOCTEXT("ChangingCurveComment", "Changing Curve Comment"), bShouldTransact);
+
+				ConditionalAction<UE::Anim::FSetCurveCommentAction>(bShouldTransact,  CurveId, Curve->Comment);
+
+				Model->CurveIdentifierToMetaData.FindChecked(CurveId).Comment = Comment;
+
+				FCurveChangedPayload Payload;
+				Payload.Identifier = CurveId;
+				Model->GetNotifier().Notify(EAnimDataModelNotifyType::CurveCommentChanged, Payload);
+
+				return true;
+			}
+			else
+			{
+				const FString CurveTypeAsString = GetCurveTypeValueName(CurveId.CurveType);
+				ReportWarningf(LOCTEXT("UnableToFindCurveWarning", "Unable to find curve: {0} of type {1}"), FText::FromName(CurveId.CurveName), FText::FromString(CurveTypeAsString));
+			}
+		}
+		else
+		{
+			ReportWarning(LOCTEXT("NonSupportedCurveCommentSetWarning", "Changing curve comment is currently only supported for float curves"));
+		}
+	}
+	else
+	{
+		ReportWarningf(LOCTEXT("InvalidCurveIdentifier", "Invalid Curve Identifier : {0}"), FText::FromName(CurveId.CurveName));
+	}	
+
+	return false;
+}
+
+
 bool UAnimSequencerController::ScaleCurve(const FAnimationCurveIdentifier& CurveId, float Origin, float Factor, bool bShouldTransact /*= true*/)
 {
 	ValidateModel();
@@ -2208,7 +2251,7 @@ void UAnimSequencerController::UpdateWithSkeleton(USkeleton* TargetSkeleton, boo
 			RemoveUnusedControlsAndCurves();		
 
 			// Forcefully re-generate legacy data structures
-			Model->GenerateLegacyCurveData();
+			Model->RegenerateLegacyCurveData();
 		}
 		CloseBracket();
 	}	
diff --git a/Engine/Plugins/Animation/AnimationData/Source/AnimationData/Private/AnimSequencerController.h b/Engine/Plugins/Animation/AnimationData/Source/AnimationData/Private/AnimSequencerController.h
--- a/Engine/Plugins/Animation/AnimationData/Source/AnimationData/Private/AnimSequencerController.h
+++ b/Engine/Plugins/Animation/AnimationData/Source/AnimationData/Private/AnimSequencerController.h
@@ -57,7 +57,8 @@ public:
 	virtual bool SetTransformCurveKey(const FAnimationCurveIdentifier& CurveId, float Time, const FTransform& Value, bool bShouldTransact = true) override;
 	virtual bool RemoveTransformCurveKey(const FAnimationCurveIdentifier& CurveId, float Time, bool bShouldTransact = true) override;
 	virtual bool RenameCurve(const FAnimationCurveIdentifier& CurveToRenameId, const FAnimationCurveIdentifier& NewCurveId, bool bShouldTransact = true) override;
-	virtual bool SetCurveColor(const FAnimationCurveIdentifier& CurveId, FLinearColor Color, bool bShouldTransact = true) override;	
+	virtual bool SetCurveColor(const FAnimationCurveIdentifier& CurveId, FLinearColor Color, bool bShouldTransact = true) override;
+	virtual bool SetCurveComment(const FAnimationCurveIdentifier& CurveId, const FString& Comment, bool bShouldTransact = true) override;
 	virtual bool ScaleCurve(const FAnimationCurveIdentifier& CurveId, float Origin, float Factor, bool bShouldTransact = true) override;
 	virtual bool SetCurveKey(const FAnimationCurveIdentifier& CurveId, const FRichCurveKey& Key, bool bShouldTransact = true) override;	
 	virtual bool RemoveCurveKey(const FAnimationCurveIdentifier& CurveId, float Time, bool bShouldTransact = true) override;
diff --git a/Engine/Plugins/Animation/AnimationData/Source/AnimationData/Private/AnimSequencerDataModel.cpp b/Engine/Plugins/Animation/AnimationData/Source/AnimationData/Private/AnimSequencerDataModel.cpp
--- a/Engine/Plugins/Animation/AnimationData/Source/AnimationData/Private/AnimSequencerDataModel.cpp
+++ b/Engine/Plugins/Animation/AnimationData/Source/AnimationData/Private/AnimSequencerDataModel.cpp
@@ -29,12 +29,6 @@
 
 #define LOCTEXT_NAMESPACE "AnimSequencerDataModel"
 
-int32 UAnimationSequencerDataModel::RetainFloatCurves = 0;
-static FAutoConsoleVariableRef CVarRetainFloatCurves(
-	TEXT("a.AnimSequencer.RetainFloatCurves"),
-	UAnimationSequencerDataModel::RetainFloatCurves,
-	TEXT("1 = Original FloatCurves are retained when generating transient curve data from Control Curves . 0 = FloatCurves are overriden with Control Curves"));
-
 int32 UAnimationSequencerDataModel::ValidationMode = 0;
 static FAutoConsoleVariableRef CValidationMode(
 	TEXT("a.AnimSequencer.ValidationMode"),
@@ -795,13 +789,25 @@ void UAnimationSequencerDataModel::OnNotify(const EAnimDataModelNotifyType& Noti
 
 		if (Collector.IsNotWithinBracket())
 		{
-			const TArray<EAnimDataModelNotifyType> CurveNotifyTypes = {EAnimDataModelNotifyType::CurveAdded, EAnimDataModelNotifyType::CurveChanged, EAnimDataModelNotifyType::CurveRenamed, EAnimDataModelNotifyType::CurveRemoved,
-			EAnimDataModelNotifyType::CurveFlagsChanged, EAnimDataModelNotifyType::CurveScaled, EAnimDataModelNotifyType::CurveColorChanged, EAnimDataModelNotifyType::Populated, EAnimDataModelNotifyType::Reset };
-			if(Collector.Contains(CurveNotifyTypes))
+			const TArray<EAnimDataModelNotifyType> CurveStorageNotifyTypes = {EAnimDataModelNotifyType::CurveAdded, EAnimDataModelNotifyType::CurveChanged, EAnimDataModelNotifyType::CurveRenamed, EAnimDataModelNotifyType::CurveRemoved,
+			EAnimDataModelNotifyType::CurveScaled, EAnimDataModelNotifyType::Populated, EAnimDataModelNotifyType::Reset };
+
+			if(Collector.Contains(CurveStorageNotifyTypes))
+			{
+				if(!ValidationMode)
+				{
+					RegenerateLegacyCurveData();
+				}
+				RefreshControlsAndProxy();
+				ResetCachedGUID();
+			}
+
+			const TArray<EAnimDataModelNotifyType> CurveDataNotifyTypes = {EAnimDataModelNotifyType::CurveFlagsChanged, EAnimDataModelNotifyType::CurveColorChanged, EAnimDataModelNotifyType::CurveCommentChanged};
+			if(Collector.Contains(CurveDataNotifyTypes))
 			{
 				if(!ValidationMode)
 				{
-					GenerateLegacyCurveData();
+					UpdateLegacyCurveData();
 				}
 				RefreshControlsAndProxy();
 				ResetCachedGUID();
@@ -863,11 +869,11 @@ UMovieSceneControlRigParameterSection* UAnimationSequencerDataModel::GetFKContro
 	return nullptr;
 }
 
-void UAnimationSequencerDataModel::GenerateLegacyCurveData()
+void UAnimationSequencerDataModel::RegenerateLegacyCurveData()
 {
-	QUICK_SCOPE_CYCLE_COUNTER(STAT_GenerateLegacyCurveData);
+	QUICK_SCOPE_CYCLE_COUNTER(STAT_RegenerateLegacyCurveData);
 	ValidateSequencerData();
-	
+
 	if (const UMovieSceneControlRigParameterTrack* Track = GetControlRigTrack())
 	{
 		for (const UMovieSceneSection* TrackSection : Track->GetAllSections())
@@ -876,70 +882,91 @@ void UAnimationSequencerDataModel::GenerateLegacyCurveData()
 			{
 				if (const UControlRig* ControlRig = Section->GetControlRig())
 				{
-						if(URigHierarchy* Hierarchy = ControlRig->GetHierarchy())
+					if(URigHierarchy* Hierarchy = ControlRig->GetHierarchy())
+					{
+						const TArray<FScalarParameterNameAndCurve>& ScalarCurves = Section->GetScalarParameterNamesAndCurves();
+						LegacyCurveData.FloatCurves.Empty();
+
+						Hierarchy->ForEach<FRigCurveElement>([Hierarchy, this, ScalarCurves, FrameRate = GetFrameRate()](const FRigCurveElement* CurveElement) -> bool
 						{
-							const FString SequencerSuffix(TEXT("_Sequencer"));
-							const TArray<FScalarParameterNameAndCurve>& ScalarCurves = Section->GetScalarParameterNamesAndCurves();				
-							if (RetainFloatCurves)
+							const FRigElementKey ControlKey(UFKControlRig::GetControlName(CurveElement->GetFName(), ERigElementType::Curve), ERigElementType::Control);
+							if (const FRigControlElement* Element = Hierarchy->Find<FRigControlElement>(ControlKey))
 							{
-								LegacyCurveData.FloatCurves.RemoveAll([SequencerSuffix](const FFloatCurve& FloatCurve)
+								FFloatCurve& FloatCurve = LegacyCurveData.FloatCurves.AddDefaulted_GetRef();
+								FloatCurve.SetName(CurveElement->GetFName());
+								FloatCurve.Color = Element->Settings.ShapeColor;
+
+								const FAnimationCurveIdentifier CurveId(FloatCurve.GetName(), ERawCurveTrackTypes::RCT_Float);
+								if (CurveIdentifierToMetaData.Contains(CurveId))
 								{
-								return FloatCurve.GetName().ToString().EndsWith(SequencerSuffix);
-								});
+									const FAnimationCurveMetaData& CurveMetaData = CurveIdentifierToMetaData.FindChecked(CurveId);
+									FloatCurve.SetCurveTypeFlags(CurveMetaData.Flags);
+									FloatCurve.Color = CurveMetaData.Color;
+									FloatCurve.Comment = CurveMetaData.Comment;
+								}
+
+								if (const FScalarParameterNameAndCurve* ScalarCurve = ScalarCurves.FindByPredicate([Element](const FScalarParameterNameAndCurve& Curve)
+								{
+									return Curve.ParameterName == Element->GetFName();
+								}))
+								{
+									AnimSequencerHelpers::ConvertFloatChannelToRichCurve(ScalarCurve->ParameterCurve, FloatCurve.FloatCurve, FrameRate);
+								}
 							}
-							else
-							{
-								LegacyCurveData.FloatCurves.Empty();
-							}				
+							return true;
+						});	
+					}
+					else
+					{
+						IAnimationDataController::ReportObjectErrorf(this, LOCTEXT("UnableToFindRigHierarchy", "Unable to retrieve RigHierarchy for ControlRig ({0})"), FText::FromString(ControlRig->GetPathName()));	      
+					}
+				}
+			}
+		}
+	}
+}
 
-						Hierarchy->ForEach<FRigCurveElement>([Hierarchy, this, ScalarCurves, FrameRate = GetFrameRate(), SequencerSuffix](const FRigCurveElement* CurveElement) -> bool
+void UAnimationSequencerDataModel::UpdateLegacyCurveData()
+{
+	QUICK_SCOPE_CYCLE_COUNTER(STAT_UpdateLegacyCurveData);
+	ValidateSequencerData();
+
+	if (const UMovieSceneControlRigParameterTrack* Track = GetControlRigTrack())
+	{
+		for (const UMovieSceneSection* TrackSection : Track->GetAllSections())
+		{
+			if (const UMovieSceneControlRigParameterSection* Section = Cast<UMovieSceneControlRigParameterSection>(TrackSection))
+			{
+				if (const UControlRig* ControlRig = Section->GetControlRig())
+				{
+					if(URigHierarchy* Hierarchy = ControlRig->GetHierarchy())
+					{
+						for(FFloatCurve& FloatCurve : LegacyCurveData.FloatCurves)
+						{
+							const FRigElementKey ControlKey(UFKControlRig::GetControlName(FloatCurve.GetName(), ERigElementType::Curve), ERigElementType::Control);
+							if (const FRigControlElement* Element = Hierarchy->Find<FRigControlElement>(ControlKey))
 							{
-								const FRigElementKey ControlKey(UFKControlRig::GetControlName(CurveElement->GetFName(), ERigElementType::Curve), ERigElementType::Control);
-								if (const FRigControlElement* Element = Hierarchy->Find<FRigControlElement>(ControlKey))
-								{
-									FFloatCurve& FloatCurve = LegacyCurveData.FloatCurves.AddDefaulted_GetRef();
-									if (RetainFloatCurves)
-									{
-									FloatCurve.SetName(FName(*(CurveElement->GetFName().ToString() + TEXT("_Sequencer"))));
-									}
-									else
-									{
-									FloatCurve.SetName(CurveElement->GetFName());
-									}						
-								
-									FloatCurve.Color = Element->Settings.ShapeColor;
-									
+								FloatCurve.Color = Element->Settings.ShapeColor;
+
 								const FAnimationCurveIdentifier CurveId(FloatCurve.GetName(), ERawCurveTrackTypes::RCT_Float);
-								if (!RetainFloatCurves || !FloatCurve.GetName().ToString().Contains(SequencerSuffix))
-									{
-										if (CurveIdentifierToMetaData.Contains(CurveId))
-										{
-											const FAnimationCurveMetaData& CurveMetaData = CurveIdentifierToMetaData.FindChecked(CurveId);
-											FloatCurve.SetCurveTypeFlags(CurveMetaData.Flags);
-											FloatCurve.Color = CurveMetaData.Color;
-										}
-									}							
-
-									if (const FScalarParameterNameAndCurve* ScalarCurve = ScalarCurves.FindByPredicate([Element](const FScalarParameterNameAndCurve& Curve)
-									{
-										return Curve.ParameterName == Element->GetFName();
-									}))
-									{							
-										AnimSequencerHelpers::ConvertFloatChannelToRichCurve(ScalarCurve->ParameterCurve, FloatCurve.FloatCurve, FrameRate);
-									}
+								if (const FAnimationCurveMetaData* CurveMetaData = CurveIdentifierToMetaData.Find(CurveId))
+								{
+									FloatCurve.SetCurveTypeFlags(CurveMetaData->Flags);
+									FloatCurve.Color = CurveMetaData->Color;
+									FloatCurve.Comment = CurveMetaData->Comment;
 								}
-								return true;
-							});	
-						}
-						else
-						{						
-							IAnimationDataController::ReportObjectErrorf(this, LOCTEXT("UnableToFindRigHierarchy", "Unable to retrieve RigHierarchy for ControlRig ({0})"), FText::FromString(ControlRig->GetPathName()));	      
+							}
 						}
-					}								
+					}
+					else
+					{
+						IAnimationDataController::ReportObjectErrorf(this, LOCTEXT("UnableToFindRigHierarchy", "Unable to retrieve RigHierarchy for ControlRig ({0})"), FText::FromString(ControlRig->GetPathName()));	      
+					}
 				}
 			}
 		}
 	}
+}
 
 void UAnimationSequencerDataModel::ValidateData() const
 {		
diff --git a/Engine/Plugins/Animation/AnimationData/Source/AnimationData/Private/AnimSequencerDataModel.h b/Engine/Plugins/Animation/AnimationData/Source/AnimationData/Private/AnimSequencerDataModel.h
--- a/Engine/Plugins/Animation/AnimationData/Source/AnimationData/Private/AnimSequencerDataModel.h
+++ b/Engine/Plugins/Animation/AnimationData/Source/AnimationData/Private/AnimSequencerDataModel.h
@@ -24,6 +24,9 @@ struct FAnimationCurveMetaData
 
 	UPROPERTY()
 	FLinearColor Color = FLinearColor::White;
+
+	UPROPERTY()
+	FString Comment;
 };
 
 UCLASS()
@@ -139,7 +142,8 @@ protected:
 	FAnimCurveBase* FindMutableCurveById(const FAnimationCurveIdentifier& CurveIdentifier);
 	FRichCurve* GetMutableRichCurve(const FAnimationCurveIdentifier& CurveIdentifier);
 
-	void GenerateLegacyCurveData();
+	void RegenerateLegacyCurveData();
+	void UpdateLegacyCurveData();
 
 	void ValidateData() const;
 	void ValidateSequencerData() const;
diff --git a/Engine/Source/Developer/AnimationDataController/Private/AnimDataController.cpp b/Engine/Source/Developer/AnimationDataController/Private/AnimDataController.cpp
--- a/Engine/Source/Developer/AnimationDataController/Private/AnimDataController.cpp
+++ b/Engine/Source/Developer/AnimationDataController/Private/AnimDataController.cpp
@@ -972,6 +972,47 @@ bool UAnimDataController::SetCurveColor(const FAnimationCurveIdentifier& CurveId
 	return false;
 }
 
+bool UAnimDataController::SetCurveComment(const FAnimationCurveIdentifier& CurveId, const FString& Comment, bool bShouldTransact)
+{
+	ValidateModel();
+
+	if (CurveId.IsValid())
+	{
+		if (CurveId.CurveType == ERawCurveTrackTypes::RCT_Float)
+		{
+			if (FFloatCurve* Curve = Model->FindMutableFloatCurveById(CurveId))
+			{
+				FTransaction Transaction = ConditionalTransaction(LOCTEXT("ChangingCurveComment", "Changing Curve Comment"), bShouldTransact);
+
+				ConditionalAction<UE::Anim::FSetCurveCommentAction>(bShouldTransact, CurveId, Curve->Comment);
+
+				Curve->Comment = Comment;
+
+				FCurveChangedPayload Payload;
+				Payload.Identifier = CurveId;
+				Model->GetNotifier().Notify(EAnimDataModelNotifyType::CurveCommentChanged, Payload);
+
+				return true;
+			}
+			else
+			{
+				const FString CurveTypeAsString = GetCurveTypeValueName(CurveId.CurveType);
+				ReportWarningf(LOCTEXT("UnableToFindCurveWarning", "Unable to find curve: {0} of type {1}"), FText::FromName(CurveId.CurveName), FText::FromString(CurveTypeAsString));
+			}
+		}
+		else
+		{
+			Report(ELogVerbosity::Warning, LOCTEXT("NonSupportedCurveCommentSetWarning", "Changing curve comment is currently only supported for float curves"));
+		}
+	}
+	else
+	{
+		ReportWarningf(LOCTEXT("InvalidCurveIdentifier", "Invalid Curve Identifier : {0}"), FText::FromName(CurveId.CurveName));
+	}
+
+	return false;
+}
+
 bool UAnimDataController::ScaleCurve(const FAnimationCurveIdentifier& CurveId, float Origin, float Factor, bool bShouldTransact /*= true*/)
 {
 	ValidateModel();
diff --git a/Engine/Source/Developer/AnimationDataController/Private/AnimDataControllerActions.cpp b/Engine/Source/Developer/AnimationDataController/Private/AnimDataControllerActions.cpp
--- a/Engine/Source/Developer/AnimationDataController/Private/AnimDataControllerActions.cpp
+++ b/Engine/Source/Developer/AnimationDataController/Private/AnimDataControllerActions.cpp
@@ -455,6 +455,21 @@ FString FSetCurveColorAction::ToStringInternal() const
 	return FText::Format(LOCTEXT("SetCurveColorAction_Description", "Setting curve color '{0}'."), FText::FromName(CurveId.CurveName)).ToString();
 }
 
+TUniquePtr<FChange> FSetCurveCommentAction::ExecuteInternal(IAnimationDataModel* Model, IAnimationDataController* Controller)
+{
+	const FAnimCurveBase& AnimationCurve = Model->GetCurve(CurveId);
+	const FString& CurrentComment = AnimationCurve.Comment;
+
+	Controller->SetCurveComment(CurveId, Comment, false);
+
+	return MakeUnique<FSetCurveCommentAction>(CurveId, CurrentComment);
+}
+
+FString FSetCurveCommentAction::ToStringInternal() const
+{
+	return FText::Format(LOCTEXT("SetCurveCommentAction_Description", "Setting curve comment '{0}'."), FText::FromName(CurveId.CurveName)).ToString();
+}
+
 FAddAtributeAction::FAddAtributeAction(const FAnimatedBoneAttribute& InAttribute) : AttributeId(InAttribute.Identifier)
 {
 	Keys = InAttribute.Curve.GetConstRefOfKeys();
diff --git a/Engine/Source/Developer/AnimationDataController/Public/AnimDataController.h b/Engine/Source/Developer/AnimationDataController/Public/AnimDataController.h
--- a/Engine/Source/Developer/AnimationDataController/Public/AnimDataController.h
+++ b/Engine/Source/Developer/AnimationDataController/Public/AnimDataController.h
@@ -59,7 +59,8 @@ public:
 	virtual bool SetTransformCurveKey(const FAnimationCurveIdentifier& CurveId, float Time, const FTransform& Value, bool bShouldTransact = true) override;
 	virtual bool RemoveTransformCurveKey(const FAnimationCurveIdentifier& CurveId, float Time, bool bShouldTransact = true) override;
 	virtual bool RenameCurve(const FAnimationCurveIdentifier& CurveToRenameId, const FAnimationCurveIdentifier& NewCurveId, bool bShouldTransact = true) override;
-	virtual bool SetCurveColor(const FAnimationCurveIdentifier& CurveId, FLinearColor Color, bool bShouldTransact = true) override;	
+	virtual bool SetCurveColor(const FAnimationCurveIdentifier& CurveId, FLinearColor Color, bool bShouldTransact = true) override;
+	virtual bool SetCurveComment(const FAnimationCurveIdentifier& CurveId, const FString& Comment, bool bShouldTransact = true) override;
 	virtual bool ScaleCurve(const FAnimationCurveIdentifier& CurveId, float Origin, float Factor, bool bShouldTransact = true) override;
 	virtual bool SetCurveKey(const FAnimationCurveIdentifier& CurveId, const FRichCurveKey& Key, bool bShouldTransact = true) override;	
 	virtual bool RemoveCurveKey(const FAnimationCurveIdentifier& CurveId, float Time, bool bShouldTransact = true) override;
diff --git a/Engine/Source/Developer/AnimationDataController/Public/AnimDataControllerActions.h b/Engine/Source/Developer/AnimationDataController/Public/AnimDataControllerActions.h
--- a/Engine/Source/Developer/AnimationDataController/Public/AnimDataControllerActions.h
+++ b/Engine/Source/Developer/AnimationDataController/Public/AnimDataControllerActions.h
@@ -342,6 +342,21 @@ protected:
 	FLinearColor Color;
 };
 
+class ANIMATIONDATACONTROLLER_API FSetCurveCommentAction : public FAnimDataBaseAction
+{
+public:
+	explicit FSetCurveCommentAction(const FAnimationCurveIdentifier& InCurveId, const FString& InComment) : CurveId(InCurveId), Comment(InComment) {}
+	virtual ~FSetCurveCommentAction() {}
+protected:
+	FSetCurveCommentAction() {}
+	virtual TUniquePtr<FChange> ExecuteInternal(IAnimationDataModel* Model, IAnimationDataController* Controller) override;
+	virtual FString ToStringInternal() const override;
+
+protected:
+	FAnimationCurveIdentifier CurveId;
+	FString Comment;
+};
+
 class ANIMATIONDATACONTROLLER_API FAddAtributeAction : public FAnimDataBaseAction
 {
 public:
diff --git a/Engine/Source/Editor/Persona/Private/AnimSequenceTimelineCommands.cpp b/Engine/Source/Editor/Persona/Private/AnimSequenceTimelineCommands.cpp
--- a/Engine/Source/Editor/Persona/Private/AnimSequenceTimelineCommands.cpp
+++ b/Engine/Source/Editor/Persona/Private/AnimSequenceTimelineCommands.cpp
@@ -17,6 +17,7 @@ void FAnimSequenceTimelineCommands::RegisterCommands()
 	UI_COMMAND(AddCurve, "Add Curve...", "Add a new variable float curve", EUserInterfaceActionType::Button, FInputChord());
 	UI_COMMAND(EditCurve, "Edit Curve", "Edit this curve in the curve editor tab", EUserInterfaceActionType::Button, FInputChord());
 	UI_COMMAND(ShowCurveKeys, "Show Curve Keys", "Show keys for all curves in the timeline", EUserInterfaceActionType::ToggleButton, FInputChord());
+	UI_COMMAND(UseTreeView, "Use Tree View", "Whether to use tree view for animation curves. This will treat period-seperated curves as seperate tree entries.", EUserInterfaceActionType::ToggleButton, FInputChord());
 	UI_COMMAND(AddMetadata, "Add Metadata...", "Add a new constant (metadata) float curve", EUserInterfaceActionType::Button, FInputChord());
 	UI_COMMAND(ConvertCurveToMetaData, "Convert To Metadata", "Convert this curve to a constant (metadata) curve.\nThis is a destructive operation and will remove all keys from this curve.", EUserInterfaceActionType::Button, FInputChord());
 	UI_COMMAND(ConvertMetaDataToCurve, "Convert To Curve", "Convert this metadata curve to a full curve", EUserInterfaceActionType::Button, FInputChord());
@@ -31,6 +32,7 @@ void FAnimSequenceTimelineCommands::RegisterCommands()
 	UI_COMMAND(SnapToNotifies, "Notifies", "Snap to notifies and notify states", EUserInterfaceActionType::ToggleButton, FInputChord());
 	UI_COMMAND(SnapToMontageSections, "Montage Sections", "Snap to montage sections", EUserInterfaceActionType::ToggleButton, FInputChord());
 	UI_COMMAND(SnapToCompositeSegments, "Composite Segments", "Snap to composite segments", EUserInterfaceActionType::ToggleButton, FInputChord());
+	UI_COMMAND(AddComment, "Add Comment", "Add a comment to this curve", EUserInterfaceActionType::ToggleButton, FInputChord());
 }
 
 #undef LOCTEXT_NAMESPACE
diff --git a/Engine/Source/Editor/Persona/Private/AnimSequenceTimelineCommands.h b/Engine/Source/Editor/Persona/Private/AnimSequenceTimelineCommands.h
--- a/Engine/Source/Editor/Persona/Private/AnimSequenceTimelineCommands.h
+++ b/Engine/Source/Editor/Persona/Private/AnimSequenceTimelineCommands.h
@@ -45,6 +45,8 @@ public:
 
 	TSharedPtr<FUICommandInfo> ShowCurveKeys;
 
+	TSharedPtr<FUICommandInfo> UseTreeView;
+
 	TSharedPtr<FUICommandInfo> AddMetadata;
 
 	TSharedPtr<FUICommandInfo> ConvertCurveToMetaData;
@@ -72,6 +74,8 @@ public:
 	TSharedPtr<FUICommandInfo> SnapToMontageSections;
 
 	TSharedPtr<FUICommandInfo> SnapToCompositeSegments;
+
+	TSharedPtr<FUICommandInfo> AddComment;
 public:
 	virtual void RegisterCommands() override;
 };
diff --git a/Engine/Source/Editor/Persona/Private/AnimTimeline/AnimModel.cpp b/Engine/Source/Editor/Persona/Private/AnimTimeline/AnimModel.cpp
--- a/Engine/Source/Editor/Persona/Private/AnimTimeline/AnimModel.cpp
+++ b/Engine/Source/Editor/Persona/Private/AnimTimeline/AnimModel.cpp
@@ -137,7 +137,7 @@ void FAnimModel::HandleWorkingRangeChanged(TRange<double> InRange)
 	WorkingRange = InRange;
 }
 
-bool FAnimModel::IsTrackSelected(const TSharedRef<FAnimTimelineTrack>& InTrack) const
+bool FAnimModel::IsTrackSelected(const TSharedRef<const FAnimTimelineTrack>& InTrack) const
 { 
 	return SelectedTracks.Find(InTrack) != nullptr;
 }
diff --git a/Engine/Source/Editor/Persona/Private/AnimTimeline/AnimModel.h b/Engine/Source/Editor/Persona/Private/AnimTimeline/AnimModel.h
--- a/Engine/Source/Editor/Persona/Private/AnimTimeline/AnimModel.h
+++ b/Engine/Source/Editor/Persona/Private/AnimTimeline/AnimModel.h
@@ -117,7 +117,7 @@ public:
 	TSharedRef<FUICommandList> GetCommandList() const { return WeakCommandList.Pin().ToSharedRef(); }
 
 	/** Get whether a track is selected */
-	bool IsTrackSelected(const TSharedRef<FAnimTimelineTrack>& InTrack) const;
+	bool IsTrackSelected(const TSharedRef<const FAnimTimelineTrack>& InTrack) const;
 
 	/** Clear all track selection */
 	void ClearTrackSelection();
diff --git a/Engine/Source/Editor/Persona/Private/AnimTimeline/AnimTimelineTrack.h b/Engine/Source/Editor/Persona/Private/AnimTimeline/AnimTimelineTrack.h
--- a/Engine/Source/Editor/Persona/Private/AnimTimeline/AnimTimelineTrack.h
+++ b/Engine/Source/Editor/Persona/Private/AnimTimeline/AnimTimelineTrack.h
@@ -152,7 +152,7 @@ public:
 	virtual void AddToContextMenu(FMenuBuilder& InMenuBuilder, TSet<FName>& InOutExistingMenuTypes) const;
 
 	/** Get the height of this track */
-	float GetHeight() const { return Height; }
+	virtual float GetHeight() const { return Height; }
 
 	/** Set the height of this track */
 	void SetHeight(float InHeight) { Height = InHeight; }
diff --git a/Engine/Source/Editor/Persona/Private/AnimTimeline/AnimTimelineTrack_Curves.cpp b/Engine/Source/Editor/Persona/Private/AnimTimeline/AnimTimelineTrack_Curves.cpp
--- a/Engine/Source/Editor/Persona/Private/AnimTimeline/AnimTimelineTrack_Curves.cpp
+++ b/Engine/Source/Editor/Persona/Private/AnimTimeline/AnimTimelineTrack_Curves.cpp
@@ -132,6 +132,19 @@ TSharedRef<SWidget> FAnimTimelineTrack_Curves::BuildCurvesSubMenu()
 			NAME_None,
 			EUserInterfaceActionType::ToggleButton
 		);
+
+		MenuBuilder.AddMenuEntry(
+			FAnimSequenceTimelineCommands::Get().UseTreeView->GetLabel(),
+			FAnimSequenceTimelineCommands::Get().UseTreeView->GetDescription(),
+			FAnimSequenceTimelineCommands::Get().UseTreeView->GetIcon(),
+			FUIAction(
+				FExecuteAction::CreateSP(this, &FAnimTimelineTrack_Curves::HandleUseTreeView),
+				FCanExecuteAction(),
+				FIsActionChecked::CreateSP(this, &FAnimTimelineTrack_Curves::IsUseTreeViewEnabled)
+			),
+			NAME_None,
+			EUserInterfaceActionType::ToggleButton
+		);
 	}
 	MenuBuilder.EndSection();
 
@@ -288,6 +301,17 @@ bool FAnimTimelineTrack_Curves::IsShowCurvePointsEnabled() const
 	return GetDefault<UPersonaOptions>()->bTimelineDisplayCurveKeys;
 }
 
+void FAnimTimelineTrack_Curves::HandleUseTreeView()
+{
+	GetMutableDefault<UPersonaOptions>()->bUseTreeViewForAnimationCurves = !GetDefault<UPersonaOptions>()->bUseTreeViewForAnimationCurves;
+	GetModel()->RefreshTracks();
+}
+
+bool FAnimTimelineTrack_Curves::IsUseTreeViewEnabled() const
+{
+	return GetDefault<UPersonaOptions>()->bUseTreeViewForAnimationCurves;
+}
+
 void FAnimTimelineTrack_Curves::OnMetadataCurveNamePicked(const FName& InCurveName)
 {
 	FSlateApplication::Get().DismissAllMenus();
diff --git a/Engine/Source/Editor/Persona/Private/AnimTimeline/AnimTimelineTrack_Curves.h b/Engine/Source/Editor/Persona/Private/AnimTimeline/AnimTimelineTrack_Curves.h
--- a/Engine/Source/Editor/Persona/Private/AnimTimeline/AnimTimelineTrack_Curves.h
+++ b/Engine/Source/Editor/Persona/Private/AnimTimeline/AnimTimelineTrack_Curves.h
@@ -33,6 +33,10 @@ private:
 	void HandleShowCurvePoints();
 	bool IsShowCurvePointsEnabled() const;
 
+	/** Handlers for switching to and from tree view */
+	void HandleUseTreeView();
+	bool IsUseTreeViewEnabled() const;
+	
 	/** Curve Picker Callbacks */
 	void OnMetadataCurveNamePicked(const FName& InCurveName);
 	void OnVariableCurveNamePicked(const FName& InCurveName);
diff --git a/Engine/Source/Editor/Persona/Private/AnimTimeline/AnimTimelineTrack_FloatCurve.cpp b/Engine/Source/Editor/Persona/Private/AnimTimeline/AnimTimelineTrack_FloatCurve.cpp
--- a/Engine/Source/Editor/Persona/Private/AnimTimeline/AnimTimelineTrack_FloatCurve.cpp
+++ b/Engine/Source/Editor/Persona/Private/AnimTimeline/AnimTimelineTrack_FloatCurve.cpp
@@ -22,6 +22,7 @@
 #include "AnimTimelineClipboard.h"
 #include "Animation/AnimData/AnimDataModel.h"
 #include "AnimTimeline/SAnimOutlinerItem.h"
+#include "Fonts/FontMeasure.h"
 
 #define LOCTEXT_NAMESPACE "FAnimTimelineTrack_FloatCurve"
 
@@ -43,36 +44,73 @@ TSharedRef<SWidget> FAnimTimelineTrack_FloatCurve::MakeTimelineWidgetContainer()
 	// zoom to fit now we have a view
 	CurveEditor->ZoomToFit(EAxisList::Y);
 
-	FLinearColor CurveColor = FloatCurve->GetCurveTypeFlag(AACF_Metadata) ? FloatCurve->GetColor().Desaturate(0.25f) : FloatCurve->GetColor().Desaturate(0.75f);
-	auto ColorLambda = [this, CurveColor]()
-	{
-		if(GetModel()->IsTrackSelected(AsShared()))
-		{
-			return FAppStyle::GetSlateColor("SelectionColor").GetSpecifiedColor().CopyWithNewOpacity(0.75f);
-		}
-		else
-		{
-			return CurveColor;
-		}
-	};
-
 	return
-		SAssignNew(TimelineWidgetContainer, SBorder)
-		.Padding(0.0f)
-		.BorderImage(FloatCurve->GetCurveTypeFlag(AACF_Metadata) ? FAppStyle::GetBrush("Sequencer.Section.SelectedSectionOverlay") : FAppStyle::GetBrush("AnimTimeline.Outliner.DefaultBorder"))
-		.BorderBackgroundColor_Lambda(ColorLambda)
+		SNew(SVerticalBox)
+		+SVerticalBox::Slot()
+		.AutoHeight()
 		[
-			CurveWidget
+			SNew(SBox)
+			.HeightOverride(Height)
+			[
+				SAssignNew(TimelineWidgetContainer, SBorder)
+				.Padding(0.0f)
+				.BorderImage(FloatCurve->GetCurveTypeFlag(AACF_Metadata) ? FAppStyle::GetBrush("Sequencer.Section.SelectedSectionOverlay") : FAppStyle::GetBrush("AnimTimeline.Outliner.DefaultBorder"))
+				.BorderBackgroundColor(this, &FAnimTimelineTrack_FloatCurve::GetTrackColor, false)
+				[
+					CurveWidget
+				]
+			]
+		]
+		+SVerticalBox::Slot()
+		.AutoHeight()
+		[
+			SNew(SBox)
+			.HeightOverride(this, &FAnimTimelineTrack_FloatCurve::GetCommentSize)
+			[
+				SNew(SBorder)
+				.Padding(5.0f, 0.0f)
+				.Visibility(this, &FAnimTimelineTrack_FloatCurve::GetCommentVisibility)
+				.ToolTipText(this, &FAnimTimelineTrack_FloatCurve::GetToolTipText)
+				.BorderImage(FAppStyle::GetBrush("AnimTimeline.Outliner.DefaultBorder"))
+				.BorderBackgroundColor(this, &FAnimTimelineTrack_FloatCurve::GetTrackColor, true)
+				[
+					SAssignNew(EditableTextComment, SInlineEditableTextBlock)
+					.Text(this, &FAnimTimelineTrack_FloatCurve::GetCommentText)
+					.IsSelected(this, &FAnimTimelineTrack_FloatCurve::IsSelected)
+					.OnTextCommitted(this, &FAnimTimelineTrack_FloatCurve::OnCommitCurveComment)
+				]
+			]
 		];
 }
 
+float FAnimTimelineTrack_FloatCurve::GetCommentHeight() const
+{
+	if(FloatCurve->Comment.IsEmpty())
+	{
+		return 0.0f;
+	}
+
+	const TSharedRef<FSlateFontMeasure> FontMeasureService = FSlateApplication::Get().GetRenderer()->GetFontMeasureService();
+	UE::Slate::FDeprecateVector2DResult Result = FontMeasureService->Measure(FloatCurve->Comment, FCoreStyle::Get().GetWidgetStyle<FTextBlockStyle>( "NormalText" ).Font);
+	return Result.Y + 6.0f;
+}
+
+FOptionalSize FAnimTimelineTrack_FloatCurve::GetCommentSize() const
+{
+	return GetCommentHeight();
+}
+
+float FAnimTimelineTrack_FloatCurve::GetHeight() const
+{
+	return Height + GetCommentHeight();
+}
+
 TSharedRef<SWidget> FAnimTimelineTrack_FloatCurve::GenerateContainerWidgetForOutliner(const TSharedRef<SAnimOutlinerItem>& InRow)
 {
 	TSharedPtr<SBorder> OuterBorder;
 	TSharedPtr<SHorizontalBox> InnerHorizontalBox;
 	TSharedRef<SWidget> OutlinerWidget = GenerateStandardOutlinerWidget(InRow, false, OuterBorder, InnerHorizontalBox);
 
-
 	UAnimMontage* AnimMontage = Cast<UAnimMontage>(GetModel()->GetAnimSequenceBase());
 	bool bChildAnimMontage = AnimMontage && AnimMontage->HasParentAsset();
 
@@ -85,7 +123,7 @@ TSharedRef<SWidget> FAnimTimelineTrack_FloatCurve::GenerateContainerWidgetForOut
 			SAssignNew(EditableTextLabel, SInlineEditableTextBlock)
 			.IsReadOnly(bChildAnimMontage)
 			.Text(this, &FAnimTimelineTrack_FloatCurve::GetLabel)
-			.IsSelected_Lambda([this](){ return GetModel()->IsTrackSelected(SharedThis(this)); })
+			.IsSelected(this, &FAnimTimelineTrack_FloatCurve::IsSelected)
 			.OnTextCommitted(this, &FAnimTimelineTrack_FloatCurve::OnCommitCurveName)
 			.HighlightText(InRow->GetHighlightText())
 		];
@@ -144,6 +182,18 @@ TSharedRef<SWidget> FAnimTimelineTrack_FloatCurve::BuildCurveTrackMenu()
 				FExecuteAction::CreateSP(this, &FAnimTimelineTrack_FloatCurve::RemoveCurve)
 			)
 		);
+
+		MenuBuilder.AddMenuEntry(
+			MakeAttributeLambda([this]()
+			{
+				return FloatCurve->Comment.IsEmpty() ? FAnimSequenceTimelineCommands::Get().AddComment->GetLabel() : LOCTEXT("EditComment", "Edit Comment");
+			}),
+			FAnimSequenceTimelineCommands::Get().AddComment->GetDescription(),
+			FAnimSequenceTimelineCommands::Get().AddComment->GetIcon(),
+			FUIAction(
+				FExecuteAction::CreateSP(this, &FAnimTimelineTrack_FloatCurve::HandleAddComment)
+			)
+		);
 	}
 	MenuBuilder.EndSection();
 
@@ -352,4 +402,58 @@ void FAnimTimelineTrack_FloatCurve::GetCurveEditInfo(int32 InCurveIndex, FName&
 	OutCurveIndex = InCurveIndex;
 }
 
+void FAnimTimelineTrack_FloatCurve::HandleAddComment()
+{
+	if(FloatCurve->Comment.IsEmpty())
+	{
+		UAnimSequenceBase* AnimSequenceBase = GetModel()->GetAnimSequenceBase();
+		AnimSequenceBase->GetController().SetCurveComment(CurveId, LOCTEXT("DefaultComment", "Comment").ToString());
+	}
+
+	ExecuteOnGameThread(UE_SOURCE_LOCATION, [WeakThis = TWeakPtr<FAnimTimelineTrack_FloatCurve>(SharedThis(this))]()
+	{
+		if(TSharedPtr<FAnimTimelineTrack_FloatCurve> This = WeakThis.Pin())
+		{
+			This->EditableTextComment->EnterEditingMode();
+		}
+	});
+}
+
+void FAnimTimelineTrack_FloatCurve::OnCommitCurveComment(const FText& InText, ETextCommit::Type CommitInfo)
+{
+	if(FloatCurve->Comment != InText.ToString())
+	{
+		UAnimSequenceBase* AnimSequenceBase = GetModel()->GetAnimSequenceBase();
+		AnimSequenceBase->GetController().SetCurveComment(CurveId, InText.ToString());
+	}
+}
+
+FText FAnimTimelineTrack_FloatCurve::GetCommentText() const
+{
+	return FText::FromString(FloatCurve->Comment);
+}
+
+EVisibility FAnimTimelineTrack_FloatCurve::GetCommentVisibility() const
+{
+	return FloatCurve->Comment.IsEmpty() ? EVisibility::Collapsed : EVisibility::Visible;
+}
+
+bool FAnimTimelineTrack_FloatCurve::IsSelected() const
+{
+	return GetModel()->IsTrackSelected(SharedThis(this));
+}
+
+FSlateColor FAnimTimelineTrack_FloatCurve::GetTrackColor(bool bForComment) const
+{
+	if(GetModel()->IsTrackSelected(AsShared()))
+	{
+		return FAppStyle::GetSlateColor("SelectionColor").GetSpecifiedColor().CopyWithNewOpacity(bForComment ? 0.5f : 0.75f);
+	}
+	else
+	{
+		FLinearColor CurveColor = FloatCurve->GetCurveTypeFlag(AACF_Metadata) ? FloatCurve->GetColor().Desaturate(0.25f) : FloatCurve->GetColor().Desaturate(0.75f);
+		return bForComment ? CurveColor.CopyWithNewOpacity(CurveColor.A * 0.5f) : CurveColor;
+	}
+}
+
 #undef LOCTEXT_NAMESPACE
diff --git a/Engine/Source/Editor/Persona/Private/AnimTimeline/AnimTimelineTrack_FloatCurve.h b/Engine/Source/Editor/Persona/Private/AnimTimeline/AnimTimelineTrack_FloatCurve.h
--- a/Engine/Source/Editor/Persona/Private/AnimTimeline/AnimTimelineTrack_FloatCurve.h
+++ b/Engine/Source/Editor/Persona/Private/AnimTimeline/AnimTimelineTrack_FloatCurve.h
@@ -28,6 +28,7 @@ public:
 	virtual void GetCurveEditInfo(int32 InCurveIndex, FName& OutName, ERawCurveTrackTypes& OutType, int32& OutCurveIndex) const override;
 	virtual bool SupportsCopy() const override { return true; }
 	virtual void Copy(UAnimTimelineClipboardContent* InOutClipboard) const override;
+	virtual float GetHeight() const override;
 	
 	/** Access the curve we are editing */
 	const FFloatCurve* GetFloatCurve() { return FloatCurve; }
@@ -51,6 +52,14 @@ private:
 	void ConvertMetaDataToCurve();
 	void RemoveCurve();
 	void OnCommitCurveName(const FText& InText, ETextCommit::Type CommitInfo);
+	float GetCommentHeight() const;
+	FOptionalSize GetCommentSize() const;
+	void HandleAddComment();
+	void OnCommitCurveComment(const FText& InText, ETextCommit::Type CommitInfo);
+	FText GetCommentText() const;
+	EVisibility GetCommentVisibility() const;
+	bool IsSelected() const;
+	FSlateColor GetTrackColor(bool bForComment) const;
 
 private:
 	/** The curve we are editing */
@@ -62,4 +71,7 @@ private:
 
 	/** Label we can edit */
 	TSharedPtr<SInlineEditableTextBlock> EditableTextLabel;
+
+	/** Comment we can edit */
+	TSharedPtr<SInlineEditableTextBlock> EditableTextComment;
 };
diff --git a/Engine/Source/Runtime/CoreUObject/Private/UObject/CoreRedirects.cpp b/Engine/Source/Runtime/CoreUObject/Private/UObject/CoreRedirects.cpp
--- a/Engine/Source/Runtime/CoreUObject/Private/UObject/CoreRedirects.cpp
+++ b/Engine/Source/Runtime/CoreUObject/Private/UObject/CoreRedirects.cpp
@@ -1934,8 +1934,6 @@ static void RegisterNativeRedirects46(TArray<FCoreRedirect>& Redirects)
 	FUNCTION_REDIRECT("SceneComponent.AttachTo", "SceneComponent.K2_AttachTo");
 	FUNCTION_REDIRECT("SkyLightComponent.SetBrightness", "SkyLightComponent.SetIntensity");
 
-	PROPERTY_REDIRECT("AnimCurveBase.CurveName", "LastObservedName");
-
 	// 4.6
 
 	CLASS_REDIRECT("ControlPointMeshComponent", "/Script/Landscape.ControlPointMeshComponent");
diff --git a/Engine/Source/Runtime/Engine/Classes/Animation/AnimData/AnimDataNotifications.h b/Engine/Source/Runtime/Engine/Classes/Animation/AnimData/AnimDataNotifications.h
--- a/Engine/Source/Runtime/Engine/Classes/Animation/AnimData/AnimDataNotifications.h
+++ b/Engine/Source/Runtime/Engine/Classes/Animation/AnimData/AnimDataNotifications.h
@@ -60,6 +60,9 @@ enum class EAnimDataModelNotifyType : uint8
 	/** Indicates a curve its color has changed. Type of payload: FCurveChangedPayload */
 	CurveColorChanged,
 
+	/** Indicates a curve has been removed. Type of payload: FCurveChangedPayload */
+	CurveCommentChanged,
+
 	/** Indicates a new attribute has been added. Type of payload: FAttributeAddedPayload */
 	AttributeAdded,
 	
diff --git a/Engine/Source/Runtime/Engine/Classes/Animation/AnimData/IAnimationDataController.h b/Engine/Source/Runtime/Engine/Classes/Animation/AnimData/IAnimationDataController.h
--- a/Engine/Source/Runtime/Engine/Classes/Animation/AnimData/IAnimationDataController.h
+++ b/Engine/Source/Runtime/Engine/Classes/Animation/AnimData/IAnimationDataController.h
@@ -408,7 +408,7 @@ public:
 	virtual bool RenameCurve(const FAnimationCurveIdentifier& CurveToRenameId, const FAnimationCurveIdentifier& NewCurveId, bool bShouldTransact = true) = 0;
 
 	/**
-	* Changes the color of the curve with provided identifier. Broadcasts a EAnimDataModelNotifyType::CurveRenamed notify if successful.
+	* Changes the color of the curve with provided identifier. Broadcasts a EAnimDataModelNotifyType::CurveColorChanged notify if successful.
 	* Currently changing curve colors is only supported for float curves.
 	*
 	* @param	CurveId				Identifier of the curve to change the color for
@@ -419,7 +419,20 @@ public:
 	*/
 	UFUNCTION(BlueprintCallable, Category = CurveData)
 	virtual bool SetCurveColor(const FAnimationCurveIdentifier& CurveId, FLinearColor Color, bool bShouldTransact = true) = 0;
-	
+
+	/**
+	* Changes the comment of the curve with provided identifier. Broadcasts a EAnimDataModelNotifyType::CurveCommentChanged notify if successful.
+	* Currently changing curve comments is only supported for float curves.
+	*
+	* @param	CurveId				Identifier of the curve to change the comment for
+	* @param	Comment				Comment to which the curve is to be set
+	* @param	bShouldTransact		Whether or not any undo-redo changes should be generated
+	*
+	* @return	Whether or not the curve comment was successfully changed
+	*/
+	UFUNCTION(BlueprintCallable, Category = CurveData)
+	virtual bool SetCurveComment(const FAnimationCurveIdentifier& CurveId, const FString& Comment, bool bShouldTransact = true) = 0;
+
 	/**
 	* Scales the curve with provided identifier. Broadcasts a EAnimDataModelNotifyType::CurveScaled notify if successful.
 	*
diff --git a/Engine/Source/Runtime/Engine/Private/Animation/AnimSequence.cpp b/Engine/Source/Runtime/Engine/Private/Animation/AnimSequence.cpp
--- a/Engine/Source/Runtime/Engine/Private/Animation/AnimSequence.cpp
+++ b/Engine/Source/Runtime/Engine/Private/Animation/AnimSequence.cpp
@@ -4440,6 +4440,7 @@ void UAnimSequence::OnModelModified(const EAnimDataModelNotifyType& NotifyType,
 		}
 
 		case EAnimDataModelNotifyType::CurveColorChanged:
+		case EAnimDataModelNotifyType::CurveCommentChanged:
 		{
 			break;
 		}
diff --git a/Engine/Source/Runtime/Engine/Public/Animation/AnimCurveTypes.h b/Engine/Source/Runtime/Engine/Public/Animation/AnimCurveTypes.h
--- a/Engine/Source/Runtime/Engine/Public/Animation/AnimCurveTypes.h
+++ b/Engine/Source/Runtime/Engine/Public/Animation/AnimCurveTypes.h
@@ -106,6 +106,9 @@ struct FAnimCurveBase
 
 	UPROPERTY()
 	FLinearColor Color;
+
+	UPROPERTY()
+	FString Comment;
 #endif
 
 private:
