Changed "New play length is same as existing one" warning to Display verbosity

diff --git a/Engine/Plugins/Animation/AnimationData/Source/AnimationData/Private/AnimSequencerController.cpp b/Engine/Plugins/Animation/AnimationData/Source/AnimationData/Private/AnimSequencerController.cpp
--- a/Engine/Plugins/Animation/AnimationData/Source/AnimationData/Private/AnimSequencerController.cpp
+++ b/Engine/Plugins/Animation/AnimationData/Source/AnimationData/Private/AnimSequencerController.cpp
@@ -156,7 +156,7 @@ void UAnimSequencerController::ResizeNumberOfFrames(FFrameNumber NewLength, FFra
 		}
 		else if (Model->bPopulated)
 		{
-			ReportWarningf(LOCTEXT("SamePlayLengthWarning", "New play length is same as existing one: {0} frames"), FText::AsNumber(NewLength.Value));
+			Reportf(ELogVerbosity::Display, LOCTEXT("SamePlayLengthWarning", "New play length is same as existing one: {0} frames"), FText::AsNumber(NewLength.Value));
 		}
 	}
 	else
@@ -203,7 +203,7 @@ void UAnimSequencerController::ResizeInFrames(FFrameNumber NewLength, FFrameNumb
 		}
 		else if (Model->bPopulated)
 		{
-			ReportWarningf(LOCTEXT("SameGetPlayLengthWarning", "New play length is same as existing one: {0} frames"), FText::AsNumber(CurrentNumberOFrames));
+			Reportf(ELogVerbosity::Display,LOCTEXT("SameGetPlayLengthWarning", "New play length is same as existing one: {0} frames"), FText::AsNumber(CurrentNumberOFrames));
 		}
 	}
 	else
diff --git a/Engine/Source/Developer/AnimationDataController/Private/AnimDataController.cpp b/Engine/Source/Developer/AnimationDataController/Private/AnimDataController.cpp
--- a/Engine/Source/Developer/AnimationDataController/Private/AnimDataController.cpp
+++ b/Engine/Source/Developer/AnimationDataController/Private/AnimDataController.cpp
@@ -147,7 +147,7 @@ void UAnimDataController::ResizeNumberOfFrames(FFrameNumber NewLength, FFrameNum
 		}
 		else if (Model->bPopulated)
 		{
-			ReportWarningf(LOCTEXT("SamePlayLengthWarning", "New play length is same as existing one: {0} frames"), FText::AsNumber(NewLength.Value));
+			Reportf(ELogVerbosity::Display,LOCTEXT("SamePlayLengthWarning", "New play length is same as existing one: {0} frames"), FText::AsNumber(NewLength.Value));
 		}
 	}
 	else
@@ -194,7 +194,7 @@ void UAnimDataController::ResizeInFrames(FFrameNumber NewLength, FFrameNumber T0
         }
 		else if (Model->bPopulated)
         {
-            ReportWarningf(LOCTEXT("SameGetPlayLengthWarning", "New play length is same as existing one: {0} frames"), FText::AsNumber(CurrentNumberOFrames));
+           Reportf(ELogVerbosity::Display,LOCTEXT("SameGetPlayLengthWarning", "New play length is same as existing one: {0} frames"), FText::AsNumber(CurrentNumberOFrames));
         }
     }
     else
