Trivial log formatting fix

diff --git a/Engine/Plugins/Animation/AnimationData/Source/AnimationData/Private/AnimSequencerDataModel.cpp b/Engine/Plugins/Animation/AnimationData/Source/AnimationData/Private/AnimSequencerDataModel.cpp
--- a/Engine/Plugins/Animation/AnimationData/Source/AnimationData/Private/AnimSequencerDataModel.cpp
+++ b/Engine/Plugins/Animation/AnimationData/Source/AnimationData/Private/AnimSequencerDataModel.cpp
@@ -1280,7 +1280,7 @@ void UAnimationSequencerDataModel::GeneratePoseData(UControlRig* ControlRig, FAn
 					}
 					else
 					{
-						IAnimationDataController::ReportObjectWarningf(this, LOCTEXT("TransformCurveBoneNotFound", "Failed to find BoneIndex for transform curve %s"), FText::FromName(CurveName));
+						IAnimationDataController::ReportObjectWarningf(this, LOCTEXT("TransformCurveBoneNotFound", "Failed to find BoneIndex for transform curve {0}"), FText::FromName(CurveName));
 					}
 				}
 			}
