AnimationSequencerDataModel evaluation always resets ref-pose to skeleton-ref pose
#fix only write out bones which are different from ref-pose (only for non matching ref-skeleton, container vs animation cases)

diff --git a/Engine/Plugins/Animation/AnimationData/Source/AnimationData/Private/AnimSequencerDataModel.cpp b/Engine/Plugins/Animation/AnimationData/Source/AnimationData/Private/AnimSequencerDataModel.cpp
--- a/Engine/Plugins/Animation/AnimationData/Source/AnimationData/Private/AnimSequencerDataModel.cpp
+++ b/Engine/Plugins/Animation/AnimationData/Source/AnimationData/Private/AnimSequencerDataModel.cpp
@@ -1218,7 +1218,9 @@ void UAnimationSequencerDataModel::GeneratePoseData(UControlRig* ControlRig, FAn
 			// Populate bone/curve elements to Pose/Curve indices
 			{
 				QUICK_SCOPE_CYCLE_COUNTER(STAT_GetMappings);
-				RigHierarchy->ForEach<FRigBoneElement>([&MeshRefSkeleton, &RequiredBones, &SkeletonRefSkeleton, &RigHierarchy, &RetargetingScope, &RigPose](const FRigBoneElement* BoneElement) -> bool
+
+				const bool bDifferentBoneContainerReferenceSkeleton = &InOutPoseData.GetPose().GetBoneContainer().GetReferenceSkeleton() != &GetSkeleton()->GetReferenceSkeleton();
+				RigHierarchy->ForEach<FRigBoneElement>([&MeshRefSkeleton, &RequiredBones, &SkeletonRefSkeleton, &RigHierarchy, &RetargetingScope, &RigPose, bDifferentBoneContainerReferenceSkeleton](FRigBoneElement* BoneElement) -> bool
 				{
 					const FName& BoneName = BoneElement->GetFName();
 					const int32 BoneIndex = MeshRefSkeleton.FindBoneIndex(BoneName);
@@ -1230,9 +1232,20 @@ void UAnimationSequencerDataModel::GeneratePoseData(UControlRig* ControlRig, FAn
 							const FCompactPoseBoneIndex CompactPoseBoneIndex = RequiredBones.GetCompactPoseIndexFromSkeletonIndex(SkeletonBoneIndex);
 							if (CompactPoseBoneIndex != INDEX_NONE)
 							{
+								const FTransform CurrentTransform = RigHierarchy->GetLocalTransform(BoneElement->GetKey());
+								if (bDifferentBoneContainerReferenceSkeleton)
+								{
+									if (CurrentTransform.Equals(RigHierarchy->GetLocalTransform(BoneElement->GetKey(), true)))
+									{
+										// In case of mismatching ref-skeletons (and therefore ref-pose) only write out transforms for bones which do not match their ref transform
+										// comes at a perf cost due to transform comparison
+										return true;
+									}
+								}
+								
 								RetargetingScope.AddTrackedBone(CompactPoseBoneIndex, SkeletonBoneIndex);
 								// Retrieve evaluated bone transform from Hierarchy
-								RigPose[CompactPoseBoneIndex] = RigHierarchy->GetLocalTransform(BoneElement->GetKey());
+								RigPose[CompactPoseBoneIndex] = CurrentTransform;
 							}
 						}
 					}
