Ensure that transform-curve lookup only checks for a valid Skeleton(bone) index and not Pose(bone) index, as the latter is valid to be missing

diff --git a/Engine/Plugins/Animation/AnimationData/Source/AnimationData/Private/AnimSequencerDataModel.cpp b/Engine/Plugins/Animation/AnimationData/Source/AnimationData/Private/AnimSequencerDataModel.cpp
--- a/Engine/Plugins/Animation/AnimationData/Source/AnimationData/Private/AnimSequencerDataModel.cpp
+++ b/Engine/Plugins/Animation/AnimationData/Source/AnimationData/Private/AnimSequencerDataModel.cpp
@@ -1239,13 +1239,16 @@ void UAnimationSequencerDataModel::GeneratePoseData(UControlRig* ControlRig, FAn
 					FTransform Value = TransformCurve.Evaluate(static_cast<float>(EvaluationContext.SampleFrameRate.AsSeconds(EvaluationContext.SampleTime)), 1.f);
 
 					const FSkeletonPoseBoneIndex SkeletonBoneIndex = FSkeletonPoseBoneIndex(SkeletonRefSkeleton.FindBoneIndex(CurveName));
-					const FCompactPoseBoneIndex BoneIndex(RequiredBones.GetCompactPoseIndexFromSkeletonPoseIndex(SkeletonBoneIndex));
-					if(BoneIndex != INDEX_NONE)
+					if (SkeletonBoneIndex != INDEX_NONE)
 					{
-						const FTransform LocalTransform = RigPose[BoneIndex];
-						RigPose[BoneIndex].SetRotation(LocalTransform.GetRotation() * Value.GetRotation());
-						RigPose[BoneIndex].SetTranslation(LocalTransform.TransformPosition(Value.GetTranslation()));
-						RigPose[BoneIndex].SetScale3D(LocalTransform.GetScale3D() * Value.GetScale3D());
+						const FCompactPoseBoneIndex BoneIndex(RequiredBones.GetCompactPoseIndexFromSkeletonPoseIndex(SkeletonBoneIndex));
+						if(BoneIndex != INDEX_NONE)
+						{
+							const FTransform LocalTransform = RigPose[BoneIndex];
+							RigPose[BoneIndex].SetRotation(LocalTransform.GetRotation() * Value.GetRotation());
+							RigPose[BoneIndex].SetTranslation(LocalTransform.TransformPosition(Value.GetTranslation()));
+							RigPose[BoneIndex].SetScale3D(LocalTransform.GetScale3D() * Value.GetScale3D());
+						}
 					}
 					else
 					{
