Sequencer: Add UpdateOrAddKeys to TMovieSceneChannelData

diff --git a/Engine/Plugins/Animation/AnimationData/Source/AnimationData/Private/AnimSequencerController.cpp b/Engine/Plugins/Animation/AnimationData/Source/AnimationData/Private/AnimSequencerController.cpp
--- a/Engine/Plugins/Animation/AnimationData/Source/AnimationData/Private/AnimSequencerController.cpp
+++ b/Engine/Plugins/Animation/AnimationData/Source/AnimationData/Private/AnimSequencerController.cpp
@@ -2851,62 +2851,118 @@ bool UAnimSequencerController::UpdateBoneCurveKeys(const FName& BoneName, const
 				const bool bContainsBoneControl = Hierarchy->Contains(BoneNameControlKey);
 				const bool bContainsParameter = Section->HasTransformParameter(BoneNameControlKey.Name);
 
-				if (bContainsBone && bContainsBoneControl && bContainsParameter)
-				{			
+				const bool bIsSorted = Algo::IsSorted(TimeValues);
+				if (bContainsBone && bContainsBoneControl && bContainsParameter && bIsSorted)
+				{
 					FTransformParameterNameAndCurves* ParameterCurvePair = Section->GetTransformParameterNamesAndCurves().FindByPredicate([BoneNameControlKey](const FTransformParameterNameAndCurves& Parameter)
 					{
 						return Parameter.ParameterName == BoneNameControlKey.Name;
 					});
 					check(ParameterCurvePair);
-					
-					const FTransform Offset = Hierarchy->GetControlOffsetTransform(Hierarchy->FindChecked<FRigControlElement>(BoneNameControlKey), ERigTransformType::InitialLocal);
 
+					const FTransform Offset = Hierarchy->GetControlOffsetTransform(Hierarchy->FindChecked<FRigControlElement>(BoneNameControlKey), ERigTransformType::InitialLocal);
 					const int32 NumKeys = Keys.Num();
+					TArray<FFrameNumber> InTimes;
+					TArray<FMovieSceneFloatValue> LocationX;
+					TArray<FMovieSceneFloatValue> LocationY;
+					TArray<FMovieSceneFloatValue> LocationZ;
+					TArray<FMovieSceneFloatValue> EulerX;
+					TArray<FMovieSceneFloatValue> EulerY;
+					TArray<FMovieSceneFloatValue> EulerZ;
+					TArray<FMovieSceneFloatValue> ScaleX;
+					TArray<FMovieSceneFloatValue> ScaleY;
+					TArray<FMovieSceneFloatValue> ScaleZ;
+					InTimes.Reserve(NumKeys);
+					LocationX.Reserve(NumKeys);
+					LocationY.Reserve(NumKeys);
+					LocationZ.Reserve(NumKeys);
+					EulerX.Reserve(NumKeys);
+					EulerY.Reserve(NumKeys);
+					EulerZ.Reserve(NumKeys);
+					ScaleX.Reserve(NumKeys);
+					ScaleY.Reserve(NumKeys);
+					ScaleZ.Reserve(NumKeys);
+
 					for (int32 KeyIndex = 0; KeyIndex < NumKeys; ++KeyIndex)
 					{
 						const FTransform& Transform = Keys[KeyIndex];
-						
+
 						FTransform OffsetTransform = UAnimationSequencerDataModel::UseDirectFKControlRigMode ? Transform : Transform.GetRelativeTransform(Offset);
 						OffsetTransform.NormalizeRotation();
-						
+
 						const FFrameNumber& FrameNumber = Model->MovieScene->GetTickResolution().AsFrameTime(TimeValues[KeyIndex]).RoundToFrame();
 						const FVector& Location = OffsetTransform.GetLocation();
 						const FVector& EulerAngles = OffsetTransform.GetRotation().Euler();
 						const FVector& Scale = OffsetTransform.GetScale3D();
-						
-						for (int32 ChannelIndex = 0; ChannelIndex < 3; ++ChannelIndex)
-						{
-							// Translation
-							AddKeyToChannel(&ParameterCurvePair->Translation[ChannelIndex], FrameNumber, static_cast<float>(Location[ChannelIndex]), EMovieSceneKeyInterpolation::Linear);
-							ParameterCurvePair->Translation[ChannelIndex].ClearDefault();
-						
-							// Rotation
-							AddKeyToChannel(&ParameterCurvePair->Rotation[ChannelIndex], FrameNumber, static_cast<float>(EulerAngles[ChannelIndex]), EMovieSceneKeyInterpolation::Linear);
-							ParameterCurvePair->Rotation[ChannelIndex].ClearDefault();
-							
-							// Scaling							
-							AddKeyToChannel(&ParameterCurvePair->Scale[ChannelIndex], FrameNumber, static_cast<float>(Scale[ChannelIndex]), EMovieSceneKeyInterpolation::Linear);
-							ParameterCurvePair->Scale[ChannelIndex].ClearDefault();
-						}
-					}
-					
+
+						FMovieSceneFloatValue LocationValueX(Location.X);
+						FMovieSceneFloatValue LocationValueY(Location.Y);
+						FMovieSceneFloatValue LocationValueZ(Location.Z);
+						LocationValueX.InterpMode = ERichCurveInterpMode::RCIM_Linear;
+						LocationValueY.InterpMode = ERichCurveInterpMode::RCIM_Linear;
+						LocationValueZ.InterpMode = ERichCurveInterpMode::RCIM_Linear;
+						FMovieSceneFloatValue EulerValueX(EulerAngles.X);
+						FMovieSceneFloatValue EulerValueY(EulerAngles.Y);
+						FMovieSceneFloatValue EulerValueZ(EulerAngles.Z);
+						EulerValueX.InterpMode = ERichCurveInterpMode::RCIM_Linear;
+						EulerValueY.InterpMode = ERichCurveInterpMode::RCIM_Linear;
+						EulerValueZ.InterpMode = ERichCurveInterpMode::RCIM_Linear;
+						FMovieSceneFloatValue ScaleValueX(Scale.X);
+						FMovieSceneFloatValue ScaleValueY(Scale.Y);
+						FMovieSceneFloatValue ScaleValueZ(Scale.Z);
+						ScaleValueX.InterpMode = ERichCurveInterpMode::RCIM_Linear;
+						ScaleValueY.InterpMode = ERichCurveInterpMode::RCIM_Linear;
+						ScaleValueZ.InterpMode = ERichCurveInterpMode::RCIM_Linear;
+						InTimes.Add(FrameNumber);
+						LocationX.Add(LocationValueX);
+						LocationY.Add(LocationValueY);
+						LocationZ.Add(LocationValueZ);
+						EulerX.Add(EulerValueX);
+						EulerY.Add(EulerValueY);
+						EulerZ.Add(EulerValueZ);
+						ScaleX.Add(ScaleValueX);
+						ScaleY.Add(ScaleValueY);
+						ScaleZ.Add(ScaleValueZ);
+					}
+
+					ParameterCurvePair->Translation[0].UpdateOrAddKeys(InTimes, LocationX);
+					ParameterCurvePair->Translation[1].UpdateOrAddKeys(InTimes, LocationY);
+					ParameterCurvePair->Translation[2].UpdateOrAddKeys(InTimes, LocationZ);
+					ParameterCurvePair->Rotation[0].UpdateOrAddKeys(InTimes, EulerX);
+					ParameterCurvePair->Rotation[1].UpdateOrAddKeys(InTimes, EulerY);
+					ParameterCurvePair->Rotation[2].UpdateOrAddKeys(InTimes, EulerZ);
+					ParameterCurvePair->Scale[0].UpdateOrAddKeys(InTimes, ScaleX);
+					ParameterCurvePair->Scale[1].UpdateOrAddKeys(InTimes, ScaleY);
+					ParameterCurvePair->Scale[2].UpdateOrAddKeys(InTimes, ScaleZ);
+					ParameterCurvePair->Translation[0].ClearDefault();
+					ParameterCurvePair->Translation[1].ClearDefault();
+					ParameterCurvePair->Translation[2].ClearDefault();
+					ParameterCurvePair->Rotation[0].ClearDefault();
+					ParameterCurvePair->Rotation[1].ClearDefault();
+					ParameterCurvePair->Rotation[2].ClearDefault();
+					ParameterCurvePair->Scale[0].ClearDefault();
+					ParameterCurvePair->Scale[1].ClearDefault();
+					ParameterCurvePair->Scale[2].ClearDefault();
 					return true;
 				}
 				else
 				{
 					if (!bContainsBone)
 					{
-						ReportWarningf(LOCTEXT("FailedToFindHierarchyBone", "Unable to find RigBone with name {0}"), FText::FromName(BoneName));						
+						ReportWarningf(LOCTEXT("FailedToFindHierarchyBone", "Unable to find RigBone with name {0}"), FText::FromName(BoneName));
 					}
-
 					if (!bContainsBoneControl)
 					{
-						ReportWarningf(LOCTEXT("BoneControlElementNotFound", "Failed to find Bone control with name {0}"), FText::FromName(BoneNameControlKey.Name));			
+						ReportWarningf(LOCTEXT("BoneControlElementNotFound", "Failed to find Bone control with name {0}"), FText::FromName(BoneNameControlKey.Name));
 					}
-					
+
 					if (!bContainsParameter)
 					{
-						ReportWarningf(LOCTEXT("BoneControlParameterNotFound", "Failed to find Bone control curve with name {0}"), FText::FromName(BoneNameControlKey.Name));			
+						ReportWarningf(LOCTEXT("BoneControlParameterNotFound", "Failed to find Bone control curve with name {0}"), FText::FromName(BoneNameControlKey.Name));
+					}
+					if (!bIsSorted)
+					{
+						ReportWarning(LOCTEXT("TimeValuesNotSorted", "Given Time Values are not sorted"));
 					}
 				}
 			}
diff --git a/Engine/Source/Runtime/MovieScene/Private/Channels/MovieSceneDoubleChannel.cpp b/Engine/Source/Runtime/MovieScene/Private/Channels/MovieSceneDoubleChannel.cpp
--- a/Engine/Source/Runtime/MovieScene/Private/Channels/MovieSceneDoubleChannel.cpp
+++ b/Engine/Source/Runtime/MovieScene/Private/Channels/MovieSceneDoubleChannel.cpp
@@ -698,6 +698,13 @@ void FMovieSceneDoubleChannel::AddKeys(const TArray<FFrameNumber>& InTimes, cons
 	AutoSetTangents();
 }
 
+void FMovieSceneDoubleChannel::UpdateOrAddKeys(const TArrayView<const FFrameNumber> InTimes, const TArrayView<FMovieSceneDoubleValue> InValues)
+{
+	GetData().UpdateOrAddKeys(InTimes, InValues);
+
+	AutoSetTangents();
+}
+
 #if WITH_EDITORONLY_DATA
 
 bool FMovieSceneDoubleChannel::GetShowCurve() const
diff --git a/Engine/Source/Runtime/MovieScene/Private/Channels/MovieSceneFloatChannel.cpp b/Engine/Source/Runtime/MovieScene/Private/Channels/MovieSceneFloatChannel.cpp
--- a/Engine/Source/Runtime/MovieScene/Private/Channels/MovieSceneFloatChannel.cpp
+++ b/Engine/Source/Runtime/MovieScene/Private/Channels/MovieSceneFloatChannel.cpp
@@ -304,6 +304,13 @@ void FMovieSceneFloatChannel::AddKeys(const TArray<FFrameNumber>& InTimes, const
 	AutoSetTangents();
 }
 
+void FMovieSceneFloatChannel::UpdateOrAddKeys(const TArrayView<const FFrameNumber> InTimes, const TArrayView<FMovieSceneFloatValue> InValues)
+{
+	GetData().UpdateOrAddKeys(InTimes, InValues);
+
+	AutoSetTangents();
+}
+
 #if WITH_EDITORONLY_DATA
 
 bool FMovieSceneFloatChannel::GetShowCurve() const
diff --git a/Engine/Source/Runtime/MovieScene/Public/Channels/MovieSceneChannelData.h b/Engine/Source/Runtime/MovieScene/Public/Channels/MovieSceneChannelData.h
--- a/Engine/Source/Runtime/MovieScene/Public/Channels/MovieSceneChannelData.h
+++ b/Engine/Source/Runtime/MovieScene/Public/Channels/MovieSceneChannelData.h
@@ -484,6 +484,66 @@ struct TMovieSceneChannelData : FMovieSceneChannelData
 		return GetHandle(ExistingKey);
 	}
 
+	/**
+	 * Updates keys with these times in the channel. Will add or insert any missing keys, and replace values of existing keys.
+	 * This assumes that the InTimes array is sorted, as it performs a tape merge to efficiently update the keys.
+	 * 
+	 * @param InTimes Times to update
+	 * @param InValues Values to update
+	 */
+	void UpdateOrAddKeys(const TArrayView<const FFrameNumber> InTimes, const TArrayView<ValueType> InValues)
+	{
+		check(InTimes.Num() == InValues.Num());
+		const int32 UpdateNum = InTimes.Num();
+		int32 CurrIdx = 0;
+		int32 UpdateIdx = 0;
+
+		while (true)
+		{
+			if (UpdateIdx == UpdateNum)
+			{
+				// No more keys to insert
+				break;
+			}
+
+			if (CurrIdx == Times->Num())
+			{
+				// Append the rest of the keys to the end
+				Times->Append(InTimes.RightChop(UpdateIdx));
+				Values->Append(InValues.RightChop(UpdateIdx));
+				break;
+			}
+			if ((*Times)[CurrIdx] == InTimes[UpdateIdx])
+			{
+				// Replace value at the associated time
+				(*Values)[CurrIdx] = InValues[UpdateIdx];
+				CurrIdx++;
+				UpdateIdx++;
+				continue;
+			}
+			if ((*Times)[CurrIdx] > InTimes[UpdateIdx])
+			{
+				// Insert new time value pair
+				Times->Insert(InTimes[UpdateIdx], CurrIdx);
+				Values->Insert(InValues[UpdateIdx], CurrIdx);
+				CurrIdx++;
+				UpdateIdx++;
+				continue;
+			}
+			if ((*Times)[CurrIdx] < InTimes[UpdateIdx])
+			{
+				// Not found insertion slot yet
+				CurrIdx++;
+				continue;
+			}
+		}
+		KeyHandles->Reset();
+		for (int32 Index = 0; Index < Times->Num(); ++Index)
+		{
+			KeyHandles->AllocateHandle(Index);
+		}
+	}
+
 	/**
 	 * Set key times for a number of keys in this channel data
 	 *
diff --git a/Engine/Source/Runtime/MovieScene/Public/Channels/MovieSceneDoubleChannel.h b/Engine/Source/Runtime/MovieScene/Public/Channels/MovieSceneDoubleChannel.h
--- a/Engine/Source/Runtime/MovieScene/Public/Channels/MovieSceneDoubleChannel.h
+++ b/Engine/Source/Runtime/MovieScene/Public/Channels/MovieSceneDoubleChannel.h
@@ -356,6 +356,8 @@ public:
 	*/
 	MOVIESCENE_API void AddKeys(const TArray<FFrameNumber>& InTimes, const TArray<FMovieSceneDoubleValue>& InValues);
 
+	MOVIESCENE_API void UpdateOrAddKeys(const TArrayView<const FFrameNumber> InTimes, const TArrayView<FMovieSceneDoubleValue> InValues);
+
 #if WITH_EDITORONLY_DATA
 	/* Get whether to show this curve in the UI */
 	MOVIESCENE_API bool GetShowCurve() const;
diff --git a/Engine/Source/Runtime/MovieScene/Public/Channels/MovieSceneFloatChannel.h b/Engine/Source/Runtime/MovieScene/Public/Channels/MovieSceneFloatChannel.h
--- a/Engine/Source/Runtime/MovieScene/Public/Channels/MovieSceneFloatChannel.h
+++ b/Engine/Source/Runtime/MovieScene/Public/Channels/MovieSceneFloatChannel.h
@@ -304,6 +304,8 @@ public:
 	*/
 	MOVIESCENE_API void AddKeys(const TArray<FFrameNumber>& InTimes, const TArray<FMovieSceneFloatValue>& InValues);
 
+	MOVIESCENE_API void UpdateOrAddKeys(const TArrayView<const FFrameNumber> InTimes, const TArrayView<FMovieSceneFloatValue> InValues);
+
 #if WITH_EDITORONLY_DATA
 	/* Get whether to show this curve in the UI */
 	MOVIESCENE_API bool GetShowCurve() const;
