Reimplemented:
"Reduce cook-time memory usage for AnimSequencerDataModel
- Lazy initialize RigHierarchy
- Reset RigHierarchy on WillNeverCacheCookedPlatformDataAgain"
Now gated behind CVAR to determine mode to operate in, disabled by default for now

diff --git a/Engine/Plugins/Animation/AnimationData/Source/AnimationData/Private/AnimSequencerController.cpp b/Engine/Plugins/Animation/AnimationData/Source/AnimationData/Private/AnimSequencerController.cpp
--- a/Engine/Plugins/Animation/AnimationData/Source/AnimationData/Private/AnimSequencerController.cpp
+++ b/Engine/Plugins/Animation/AnimationData/Source/AnimationData/Private/AnimSequencerController.cpp
@@ -1344,7 +1344,7 @@ bool UAnimSequencerController::SetCurveAttributes(const FAnimationCurveIdentifie
 
 			if (UFKControlRig* FKRig = Cast<UFKControlRig>(ControlRig))
 			{
-				const URigHierarchy* Hierarchy = FKRig->GetHierarchy();
+				const URigHierarchy* Hierarchy = Model->GetControlRigHierarchy();
 				if (Hierarchy ||  IgnoreSkeletonValidation())
 				{
 					const FRigElementKey CurveKey(URigHierarchy::GetSanitizedName(CurveId.CurveName), ERigElementType::Curve);
@@ -1580,7 +1580,7 @@ bool UAnimSequencerController::SetBoneTrackKeys(FName BoneName, const TArray<FVe
 			{
 				if (UControlRig* ControlRig = Model->GetControlRig())
 				{
-					if (URigHierarchy* Hierarchy = ControlRig->GetHierarchy())
+					if (URigHierarchy* Hierarchy = Model->GetControlRigHierarchy())
 					{
 						const FRigElementKey BoneKey(URigHierarchy::GetSanitizedName(BoneName), ERigElementType::Bone);
 						const FRigElementKey BoneNameControlKey(UFKControlRig::GetControlName(BoneKey.Name, ERigElementType::Bone), ERigElementType::Control);
@@ -2194,7 +2194,7 @@ void UAnimSequencerController::RemoveUnusedControlsAndCurves() const
 	{
 		if (UControlRig* ControlRig = Model->GetControlRig())
 		{
-			if (URigHierarchy* Hierarchy = ControlRig->GetHierarchy())
+			if (URigHierarchy* Hierarchy = Model->GetControlRigHierarchy())
 			{
 				TArray<FRigElementKey> ElementsToRemove;
 				
@@ -2427,7 +2427,7 @@ void UAnimSequencerController::InitializeModel()
 				{
 					if(UFKControlRig* ControlRig = NewObject<UFKControlRig>(Track, UFKControlRig::StaticClass(), NAME_None, RF_Transactional))
 					{
-						Model->InitializeFKControlRig(ControlRig, Model->GetSkeleton());
+						Model->InitializeFKControlRig(ControlRig, Model->GetSkeleton(), true);
 
 						// Remove all control elements (start fresh)
 						if (URigHierarchy* Hierarchy = ControlRig->GetHierarchy())
@@ -2495,7 +2495,7 @@ bool UAnimSequencerController::AddBoneControl(const FName& BoneName) const
 	{
 		if (UFKControlRig* FKRig = Cast<UFKControlRig>(Section->GetControlRig()))
 		{
-			if(URigHierarchy* Hierarchy = FKRig->GetHierarchy())
+			if(URigHierarchy* Hierarchy = Model->GetControlRigHierarchy())
 			{
 				const FRigElementKey BoneKey(URigHierarchy::GetSanitizedName(BoneName), ERigElementType::Bone);
 				const FRigElementKey BoneNameControlKey(UFKControlRig::GetControlName(BoneKey.Name, ERigElementType::Bone), ERigElementType::Control);
@@ -2600,7 +2600,7 @@ bool UAnimSequencerController::RemoveBoneControl(const FName& BoneName) const
 	{
 		if (UFKControlRig* FKRig = Cast<UFKControlRig>(Section->GetControlRig()))
 		{
-			if(URigHierarchy* Hierarchy = FKRig->GetHierarchy())
+			if(URigHierarchy* Hierarchy = Model->GetControlRigHierarchy())
 			{
 				const FRigElementKey BoneNameControlKey(UFKControlRig::GetControlName(URigHierarchy::GetSanitizedName(BoneName), ERigElementType::Bone), ERigElementType::Control);
 				
@@ -2657,7 +2657,7 @@ bool UAnimSequencerController::SetBoneCurveKeys(const FName& BoneName, const TAr
 		UControlRig* ControlRig = Section->GetControlRig();
 		if (UFKControlRig* FKRig = Cast<UFKControlRig>(ControlRig))
 		{
-			const URigHierarchy* Hierarchy = FKRig->GetHierarchy();
+			const URigHierarchy* Hierarchy = Model->GetControlRigHierarchy();
 			if(Hierarchy || IgnoreSkeletonValidation())
 			{
 				const FRigElementKey BoneKey(URigHierarchy::GetSanitizedName(BoneName), ERigElementType::Bone);
@@ -2849,7 +2849,7 @@ bool UAnimSequencerController::UpdateBoneCurveKeys(const FName& BoneName, const
 		UControlRig* ControlRig = Section->GetControlRig();
 		if (UFKControlRig* FKRig = Cast<UFKControlRig>(ControlRig))
 		{
-			if(URigHierarchy* Hierarchy = FKRig->GetHierarchy())
+			if(URigHierarchy* Hierarchy = Model->GetControlRigHierarchy())
 			{
 				const FRigElementKey BoneKey(URigHierarchy::GetSanitizedName(BoneName), ERigElementType::Bone);
 				const FRigElementKey BoneNameControlKey(UFKControlRig::GetControlName(BoneKey.Name, ERigElementType::Bone), ERigElementType::Control);
@@ -3004,7 +3004,7 @@ bool UAnimSequencerController::RemoveBoneCurveKey(const FName& BoneName, float T
 	
 			if (UFKControlRig* FKRig = Cast<UFKControlRig>(ControlRig))
 			{
-				if(const URigHierarchy* Hierarchy = FKRig->GetHierarchy())
+				if(const URigHierarchy* Hierarchy = Model->GetControlRigHierarchy())
 				{
 					const FRigElementKey BoneKey(URigHierarchy::GetSanitizedName(BoneName), ERigElementType::Bone);
 					const FRigElementKey BoneNameControlKey(UFKControlRig::GetControlName(BoneKey.Name, ERigElementType::Bone), ERigElementType::Control);
@@ -3077,7 +3077,7 @@ bool UAnimSequencerController::AddCurveControl(const FName& CurveName) const
 		UControlRig* ControlRig = Section->GetControlRig();
 		if (UFKControlRig* FKRig = Cast<UFKControlRig>(ControlRig))
 		{
-			if(URigHierarchy* Hierarchy = FKRig->GetHierarchy())
+			if(URigHierarchy* Hierarchy = Model->GetControlRigHierarchy())
 			{				
 				const FRigElementKey CurveKey(URigHierarchy::GetSanitizedName(CurveName), ERigElementType::Curve);
 				const FRigElementKey CurveControlKey(UFKControlRig::GetControlName(CurveKey.Name, ERigElementType::Curve), ERigElementType::Control);
@@ -3159,7 +3159,7 @@ bool UAnimSequencerController::RenameCurveControl(const FName& CurveName, const
 	{
 		if (UControlRig* ControlRig = Section->GetControlRig())
 		{
-			if (URigHierarchy* RigHierarchy = ControlRig->GetHierarchy())
+			if (URigHierarchy* RigHierarchy = Model->GetControlRigHierarchy())
 			{
 				if (URigHierarchyController* Controller = RigHierarchy->GetController())
 				{
@@ -3248,7 +3248,7 @@ bool UAnimSequencerController::RemoveCurveControl(const FName& CurveName) const
 
 		if (UFKControlRig* FKRig = Cast<UFKControlRig>(ControlRig))
 		{
-			if(URigHierarchy* Hierarchy = FKRig->GetHierarchy())
+			if(URigHierarchy* Hierarchy = Model->GetControlRigHierarchy())
 			{									
 				const FRigElementKey CurveKey(URigHierarchy::GetSanitizedName(CurveName), ERigElementType::Curve);
 				const FRigElementKey CurveControlKey(UFKControlRig::GetControlName(CurveKey.Name, ERigElementType::Curve), ERigElementType::Control);
@@ -3312,7 +3312,7 @@ bool UAnimSequencerController::SetCurveControlKeys(const FName& CurveName, const
 
 		if (UFKControlRig* FKRig = Cast<UFKControlRig>(ControlRig))
 		{
-			const URigHierarchy* Hierarchy = FKRig->GetHierarchy();
+			const URigHierarchy* Hierarchy = Model->GetControlRigHierarchy();
 
 			if (Hierarchy ||  IgnoreSkeletonValidation())
 			{
@@ -3400,7 +3400,7 @@ bool UAnimSequencerController::SetCurveControlKey(const FName& CurveName, const
 
 		if (UFKControlRig* FKRig = Cast<UFKControlRig>(ControlRig))
 		{
-			if(const URigHierarchy* Hierarchy = FKRig->GetHierarchy())
+			if(const URigHierarchy* Hierarchy = Model->GetControlRigHierarchy())
 			{
 				const FRigElementKey CurveKey(URigHierarchy::GetSanitizedName(CurveName), ERigElementType::Curve);
 				const FRigElementKey CurveControlKey(UFKControlRig::GetControlName(CurveKey.Name, ERigElementType::Curve), ERigElementType::Control);
@@ -3489,7 +3489,7 @@ bool UAnimSequencerController::RemoveCurveControlKey(const FName& CurveName, flo
 
 		if (UFKControlRig* FKRig = Cast<UFKControlRig>(ControlRig))
 		{
-			if(const URigHierarchy* Hierarchy = FKRig->GetHierarchy())
+			if(const URigHierarchy* Hierarchy = Model->GetControlRigHierarchy())
 			{
 				const FFrameNumber FrameNumber = Model->MovieScene->GetTickResolution().AsFrameNumber(Time);
 				const FRigElementKey CurveKey(URigHierarchy::GetSanitizedName(CurveName), ERigElementType::Curve);
@@ -3553,7 +3553,7 @@ bool UAnimSequencerController::DuplicateCurveControl(const FName& CurveName, con
 	{
 		if (UFKControlRig* FKRig = Cast<UFKControlRig>(Section->GetControlRig()))
 		{
-			if (URigHierarchy* Hierarchy = FKRig->GetHierarchy())
+			if (URigHierarchy* Hierarchy = Model->GetControlRigHierarchy())
 			{
 				if (URigHierarchyController* HierarchyController = Hierarchy->GetController())
 				{
diff --git a/Engine/Plugins/Animation/AnimationData/Source/AnimationData/Private/AnimSequencerDataModel.cpp b/Engine/Plugins/Animation/AnimationData/Source/AnimationData/Private/AnimSequencerDataModel.cpp
--- a/Engine/Plugins/Animation/AnimationData/Source/AnimationData/Private/AnimSequencerDataModel.cpp
+++ b/Engine/Plugins/Animation/AnimationData/Source/AnimationData/Private/AnimSequencerDataModel.cpp
@@ -42,6 +42,23 @@ static FAutoConsoleVariableRef CVarDirectControlRigMode(
 	UAnimationSequencerDataModel::UseDirectFKControlRigMode,
 	TEXT("1 = FKControl rig uses Direct method for setting Control transforms. 0 = FKControl rig uses Replace method (transform offsets) for setting Control transforms"));
 
+int32 UAnimationSequencerDataModel::LazyRigHierarchyInitializationMode = 0;
+static FAutoConsoleVariableRef CVarLazyRigHierarchyInitializationMode(
+	TEXT("a.AnimSequencer.LazyRigHierarchyInitMode"),
+	UAnimationSequencerDataModel::LazyRigHierarchyInitializationMode,
+	TEXT("0 = RigHierarchy is always initialized during PostLoad.\n1 = RigHierarchy is lazily initialized _while_ running CookCommandlet otherwise during PostLoad\n2 = RigHierarchy is always lazily initialized when required for pose evaluation or model modification _while not_ running CookCommandlet\n3 =RigHierarchy is always lazily initialized"),
+	ECVF_ReadOnly);
+
+static bool ShouldInitializeHierarchyDuringCook()
+{
+	return UAnimationSequencerDataModel::LazyRigHierarchyInitializationMode != 1 && UAnimationSequencerDataModel::LazyRigHierarchyInitializationMode != 3;
+}
+
+static bool ShouldInitializeHierarchyDuringPostLoad()
+{
+	return (ShouldInitializeHierarchyDuringCook() && IsRunningCookCommandlet()) || UAnimationSequencerDataModel::LazyRigHierarchyInitializationMode == 0;
+}
+
 void UAnimationSequencerDataModel::RemoveOutOfDateControls() const
 {
 	if (UMovieSceneControlRigParameterSection* Section = GetFKControlRigSection())
@@ -102,7 +119,36 @@ USkeleton* UAnimationSequencerDataModel::GetSkeleton() const
 	return Skeleton;
 }
 
-void UAnimationSequencerDataModel::InitializeFKControlRig(UFKControlRig* FKControlRig, USkeleton* Skeleton) const
+void UAnimationSequencerDataModel::InitializeRigHierarchy(UFKControlRig* FKControlRig, USkeleton* Skeleton) const
+{
+	FScopeLock Lock(&EvaluationLock);
+	
+	UFKControlRig::FRigElementInitializationOptions InitOptions;
+	InitOptions.bImportCurves = false;	
+	if(UMovieSceneControlRigParameterSection* Section = GetFKControlRigSection())
+	{
+		for (const FScalarParameterNameAndCurve& AnimCurve : Section->GetScalarParameterNamesAndCurves())
+		{
+			InitOptions.CurveNames.Add(UFKControlRig::GetControlTargetName(AnimCurve.ParameterName, ERigElementType::Curve));
+		}
+
+		for (const FTransformParameterNameAndCurves& BoneCurve : Section->GetTransformParameterNamesAndCurves())
+		{
+			InitOptions.BoneNames.Add(UFKControlRig::GetControlTargetName(BoneCurve.ParameterName, ERigElementType::Bone));
+		}
+	}
+	InitOptions.bGenerateBoneControls = InitOptions.BoneNames.Num() > 0;
+	FKControlRig->SetInitializationOptions(InitOptions);
+
+	FKControlRig->Initialize();
+
+	FKControlRig->SetApplyMode(UseDirectFKControlRigMode == 1 ? EControlRigFKRigExecuteMode::Direct : EControlRigFKRigExecuteMode::Replace);
+	FKControlRig->SetBoneInitialTransformsFromRefSkeleton(Skeleton->GetReferenceSkeleton());
+	FKControlRig->Evaluate_AnyThread();
+	bRigHierarchyInitialized = true;
+}
+
+void UAnimationSequencerDataModel::InitializeFKControlRig(UFKControlRig* FKControlRig, USkeleton* Skeleton, bool bForceHierarchyInitialization /*=false*/) const
 {
 	checkf(FKControlRig, TEXT("Invalid FKControlRig provided"));
 	if (Skeleton)
@@ -112,31 +158,17 @@ void UAnimationSequencerDataModel::InitializeFKControlRig(UFKControlRig* FKContr
 		FKControlRig->SetObjectBinding(MakeShared<FControlRigObjectBinding>());
 		FKControlRig->GetObjectBinding()->BindToObject(Skeleton);
 
-		UFKControlRig::FRigElementInitializationOptions InitOptions;
-		InitOptions.bImportCurves = false;	
-		if(UMovieSceneControlRigParameterSection* Section = GetFKControlRigSection())
+		if (bForceHierarchyInitialization)
 		{
-			for (const FScalarParameterNameAndCurve& AnimCurve : Section->GetScalarParameterNamesAndCurves())
-			{
-				InitOptions.CurveNames.Add(UFKControlRig::GetControlTargetName(AnimCurve.ParameterName, ERigElementType::Curve));
-			}
-
-			for (const FTransformParameterNameAndCurves& BoneCurve : Section->GetTransformParameterNamesAndCurves())
-			{
-				InitOptions.BoneNames.Add(UFKControlRig::GetControlTargetName(BoneCurve.ParameterName, ERigElementType::Bone));
-			}
+			InitializeRigHierarchy(FKControlRig, Skeleton);
 		}
-		InitOptions.bGenerateBoneControls = InitOptions.BoneNames.Num() > 0;
-		FKControlRig->SetInitializationOptions(InitOptions);
-
-		FKControlRig->Initialize();
-
-		FKControlRig->SetApplyMode(UseDirectFKControlRigMode == 1 ? EControlRigFKRigExecuteMode::Direct : EControlRigFKRigExecuteMode::Replace);
-		FKControlRig->SetBoneInitialTransformsFromRefSkeleton(Skeleton->GetReferenceSkeleton());
-		FKControlRig->Evaluate_AnyThread();
 		
 		UnlockEvaluationAndModification();
 	}
+	else
+	{
+		IAnimationDataController::ReportObjectErrorf(this, LOCTEXT("InvalidFKControlRig", "Unable to initialize FKControlRig for AnimationSequencerDataModel for {0}, provided FKControlRig is invalid"), FText::FromString(*GetAnimationSequence()->GetPathName()));
+	}	
 }
 
 UControlRig* UAnimationSequencerDataModel::GetControlRig() const
@@ -175,7 +207,9 @@ void UAnimationSequencerDataModel::PostLoad()
 		{
 			if (UFKControlRig* ControlRig = Cast<UFKControlRig>(Section->GetControlRig()))
 			{
-				InitializeFKControlRig(ControlRig, GetSkeleton());
+				// Allocate RigHierarchy UObject on game-thread
+				ControlRig->PostInitInstanceIfRequired();
+				InitializeFKControlRig(ControlRig, GetSkeleton(), ShouldInitializeHierarchyDuringPostLoad());
 			}
 		}
 
@@ -218,6 +252,11 @@ void UAnimationSequencerDataModel::PreSave(FObjectPreSaveContext ObjectSaveConte
 void UAnimationSequencerDataModel::WillNeverCacheCookedPlatformDataAgain()
 {
 	Super::WillNeverCacheCookedPlatformDataAgain();
+	// Only allow clearing hierarchy data in case we are in lazy-initialize mode
+	if (!ShouldInitializeHierarchyDuringCook())
+    {    
+	    ClearControlRigData();
+    }
 }
 
 void UAnimationSequencerDataModel::PreEditUndo()
@@ -697,6 +736,15 @@ void UAnimationSequencerDataModel::Evaluate(FAnimationPoseData& InOutPoseData, c
 	if (UMovieSceneControlRigParameterTrack* Track = GetControlRigTrack())
 	{
 		FScopeLock Lock(&EvaluationLock);
+
+		if(!bRigHierarchyInitialized)
+		{
+			if (UFKControlRig* FKControlRig = Cast<UFKControlRig>(GetControlRig()))
+			{
+				InitializeRigHierarchy(FKControlRig, GetSkeleton());
+			}
+		}
+		
 		// Evaluates and applies control curves from track to ControlRig
 		EvaluateTrack(Track, EvaluationContext);
 
@@ -825,6 +873,27 @@ UMovieSceneControlRigParameterSection* UAnimationSequencerDataModel::GetFKContro
 	return nullptr;
 }
 
+URigHierarchy* UAnimationSequencerDataModel::GetControlRigHierarchy() const
+{	
+	if (const UMovieSceneControlRigParameterSection* Section = GetFKControlRigSection())
+	{
+		UControlRig* ControlRig = Section->GetControlRig();
+		if (UFKControlRig* FKRig = Cast<UFKControlRig>(ControlRig))
+		{
+			if (!bRigHierarchyInitialized)
+			{
+				InitializeRigHierarchy(FKRig, GetSkeleton());
+			}
+
+			return FKRig->GetHierarchy();
+		}
+
+		IAnimationDataController::ReportObjectErrorf(this, LOCTEXT("UnableToFindRigHierarchy", "Unable to retrieve RigHierarchy for ControlRig ({0})"), FText::FromString(ControlRig->GetPathName()));
+	}
+
+	return nullptr;
+}
+
 void UAnimationSequencerDataModel::RegenerateLegacyCurveData()
 {
 	QUICK_SCOPE_CYCLE_COUNTER(STAT_RegenerateLegacyCurveData);
@@ -1137,6 +1206,19 @@ void UAnimationSequencerDataModel::GenerateTransformKeysForControl(const FName&
 	}, &FrameNumbers);
 }
 
+void UAnimationSequencerDataModel::ClearControlRigData()
+{
+	if (bRigHierarchyInitialized)
+	{
+		FScopeLock Lock(&EvaluationLock);
+		if (UFKControlRig* FKControlRig = Cast<UFKControlRig>(GetControlRig()))
+		{
+			FKControlRig->GetHierarchy()->Reset();
+			bRigHierarchyInitialized = false;
+		}
+	}	
+}
+
 UMovieScene* UAnimationSequencerDataModel::GetMovieScene() const
 {
 	return MovieScene;
diff --git a/Engine/Plugins/Animation/AnimationData/Source/AnimationData/Private/AnimSequencerDataModel.h b/Engine/Plugins/Animation/AnimationData/Source/AnimationData/Private/AnimSequencerDataModel.h
--- a/Engine/Plugins/Animation/AnimationData/Source/AnimationData/Private/AnimSequencerDataModel.h
+++ b/Engine/Plugins/Animation/AnimationData/Source/AnimationData/Private/AnimSequencerDataModel.h
@@ -13,6 +13,7 @@ class UControlRig;
 class UMovieSceneControlRigParameterTrack;
 class UMovieSceneControlRigParameterSection;
 class UFKControlRig;
+class URigHierarchy;
 
 USTRUCT()
 struct FAnimationCurveMetaData
@@ -38,6 +39,7 @@ public:
 	static int32 RetainFloatCurves;
 	static int32 ValidationMode;
 	static int32 UseDirectFKControlRigMode;
+	static int32 LazyRigHierarchyInitializationMode;
 	
 	/** Begin UObject overrides*/
 	virtual void GetPreloadDependencies(TArray<UObject*>& OutDeps) override;
@@ -163,8 +165,9 @@ protected:
 
 	UMovieSceneControlRigParameterTrack* GetControlRigTrack() const;
 	UMovieSceneControlRigParameterSection* GetFKControlRigSection() const;
+	URigHierarchy* GetControlRigHierarchy() const;
 	USkeleton* GetSkeleton() const;
-	void InitializeFKControlRig(UFKControlRig* FKControlRig, USkeleton* Skeleton) const;
+	void InitializeFKControlRig(UFKControlRig* FKControlRig, USkeleton* Skeleton, bool bForceHierarchyInitialization=false) const;
 	UControlRig* GetControlRig() const;
 	
 	void IterateTransformControlCurve(const FName& BoneName, TFunction<void(const FTransform&, const FFrameNumber&)> IterationFunction, const TArray<FFrameNumber>* InFrameNumbers = nullptr) const;
@@ -172,10 +175,13 @@ protected:
 	void RemoveOutOfDateControls() const;
 	
 	void GenerateTransformKeysForControl(const FName& BoneName, const TArray<FFrameNumber>& FrameNumbers, TArray<FTransform>& InOutTransforms) const;
+	void ClearControlRigData();
 private:	
 	template <typename HasherType>
 	void GenerateStateHash(HasherType& Hasher) const;
 
+	void InitializeRigHierarchy(UFKControlRig* FKControlRig, USkeleton* Skeleton) const;
+	
 	/** Dynamic delegate event allows scripting to register to any broadcast-ed notify. */
 	UPROPERTY(BlueprintAssignable, Transient, Category = AnimationDataModel, meta = (ScriptName = "ModifiedEvent", AllowPrivateAccess = "true"))
 	FAnimDataModelModifiedDynamicEvent ModifiedEventDynamic;	
@@ -204,6 +210,9 @@ private:
 	UPROPERTY(VisibleAnywhere, Category=AnimSequencer)
 	FGuid CachedRawDataGUID;
 
+	UPROPERTY(VisibleAnywhere, Transient, Category=AnimSequencer)
+	mutable bool bRigHierarchyInitialized = false;
+
 	// Scope lock to prevent contention around ControlRig->Evaluation()
 	mutable FCriticalSection EvaluationLock;
 
