- Fallback to control initial transform in case of no default value nor set keys when iterating a bone track
- Force Linked-Anim-Sequence workflow to _not_ transact its animation recording

diff --git a/Engine/Plugins/Animation/AnimationData/Source/AnimationData/Private/AnimSequencerDataModel.cpp b/Engine/Plugins/Animation/AnimationData/Source/AnimationData/Private/AnimSequencerDataModel.cpp
--- a/Engine/Plugins/Animation/AnimationData/Source/AnimationData/Private/AnimSequencerDataModel.cpp
+++ b/Engine/Plugins/Animation/AnimationData/Source/AnimationData/Private/AnimSequencerDataModel.cpp
@@ -1719,9 +1719,13 @@ void UAnimationSequencerDataModel::IterateBoneKeys(const FName& BoneName, TFunct
 
 
 		const int32 NumberOfKeysToIterate = MaxNumberOfKeys;
-		FVector3f PreviousPos = FVector3f::ZeroVector;
-		FVector3f PreviousScale = FVector3f::ZeroVector;
-		FVector3f PreviousRot = FVector3f::ZeroVector;
+		const FRigControlElement* BoneControl = Hierarchy->Find<FRigControlElement>(BoneControlKey);
+		const FTransform InitialTransform = BoneControl ? BoneControl->GetTransform().Initial.Local.Get() : FTransform::Identity;
+
+		// Initialize components with initial control values (in case there is no default value nor keys)
+		FVector3f PreviousPos = FVector3f(InitialTransform.GetLocation());
+		FVector3f PreviousScale = FVector3f(InitialTransform.GetScale3D());
+		FVector3f PreviousRot = FVector3f(InitialTransform.GetRotation().Euler());
 		FQuat4f PreviousQuat;
 
 		for (int32 KeyIndex = 0; KeyIndex < NumberOfKeysToIterate; ++KeyIndex)
@@ -1734,7 +1738,7 @@ void UAnimationSequencerDataModel::IterateBoneKeys(const FName& BoneName, TFunct
 				{
 					if (KeyIndex == 0)
 					{
-						PreviousPos[ChannelIndex] = ControlCurve.Translation[ChannelIndex].GetDefault().GetValue();
+						PreviousPos[ChannelIndex] = ControlCurve.Translation[ChannelIndex].GetDefault().Get(PreviousPos[ChannelIndex]);
 					}
 				}
 				else
@@ -1746,7 +1750,7 @@ void UAnimationSequencerDataModel::IterateBoneKeys(const FName& BoneName, TFunct
 				{
 					if (KeyIndex == 0)
 					{
-						PreviousRot[ChannelIndex] = ControlCurve.Rotation[ChannelIndex].GetDefault().GetValue();
+						PreviousRot[ChannelIndex] = ControlCurve.Rotation[ChannelIndex].GetDefault().Get(PreviousRot[ChannelIndex]);
 					}
 				}
 				else
@@ -1758,14 +1762,13 @@ void UAnimationSequencerDataModel::IterateBoneKeys(const FName& BoneName, TFunct
 				{
 					if (KeyIndex == 0)
 					{
-						PreviousScale[ChannelIndex] = ControlCurve.Scale[ChannelIndex].GetDefault().GetValue();
+						PreviousScale[ChannelIndex] = ControlCurve.Scale[ChannelIndex].GetDefault().Get(PreviousScale[ChannelIndex]);
 					}
 				}
 				else
 				{
 					ControlCurve.Scale[ChannelIndex].Evaluate(Frame, PreviousScale[ChannelIndex]);
 				}
-				
 			}
 
 			PreviousQuat = FQuat4f::MakeFromEuler(PreviousRot);
diff --git a/Engine/Source/Editor/MovieSceneTools/Private/TrackEditors/CommonAnimationTrackEditor.cpp b/Engine/Source/Editor/MovieSceneTools/Private/TrackEditors/CommonAnimationTrackEditor.cpp
--- a/Engine/Source/Editor/MovieSceneTools/Private/TrackEditors/CommonAnimationTrackEditor.cpp
+++ b/Engine/Source/Editor/MovieSceneTools/Private/TrackEditors/CommonAnimationTrackEditor.cpp
@@ -1578,6 +1578,7 @@ void FCommonAnimationTrackEditor::OnSequencerSaved(ISequencer& )
 							AnimSeqExportOption->bUseCustomFrameRate = Item.bUseCustomFrameRate;
 							AnimSeqExportOption->CustomFrameRate = Item.CustomFrameRate;
 
+							TGuardValue<bool> TransactionOptionGuard(AnimSeqExportOption->bTransactRecording, false);
 							FAnimExportSequenceParameters AESP;
 							AESP.Player = SequencerPtr.Get();
 							AESP.RootToLocalTransform = RootToLocalTransform;
