Update AddAttribute functions to not allow adding an attribute with a name that already exists

diff --git a/Engine/Plugins/Animation/AnimationData/Source/AnimationData/Private/AnimSequencerController.cpp b/Engine/Plugins/Animation/AnimationData/Source/AnimationData/Private/AnimSequencerController.cpp
--- a/Engine/Plugins/Animation/AnimationData/Source/AnimationData/Private/AnimSequencerController.cpp
+++ b/Engine/Plugins/Animation/AnimationData/Source/AnimationData/Private/AnimSequencerController.cpp
@@ -1841,12 +1841,12 @@ bool UAnimSequencerController::AddAttribute(const FAnimationAttributeIdentifier&
 {
 	if (AttributeIdentifier.IsValid())
 	{
-		const bool bAttributeAlreadyExists = Model->AnimatedBoneAttributes.ContainsByPredicate([AttributeIdentifier](const FAnimatedBoneAttribute& Attribute) -> bool
+		const FAnimatedBoneAttribute* const ExistingAttribute = Model->AnimatedBoneAttributes.FindByPredicate([AttributeIdentifier](const FAnimatedBoneAttribute& Attribute) -> bool
 		{
-			return Attribute.Identifier == AttributeIdentifier;
+			return Attribute.Identifier.GetName() == AttributeIdentifier.GetName() && Attribute.Identifier.GetBoneName() == AttributeIdentifier.GetBoneName();
 		});
 
-		if (!bAttributeAlreadyExists)
+		if (!ExistingAttribute)
 		{
 			IAnimationDataModel::FEvaluationAndModificationLock Lock(*ModelInterface);
 			FTransaction Transaction = ConditionalTransaction(LOCTEXT("AddAttribute", "Adding Animated Bone Attribute"), bShouldTransact);
@@ -1866,8 +1866,11 @@ bool UAnimSequencerController::AddAttribute(const FAnimationAttributeIdentifier&
 		}
 		else
 		{
-			ReportErrorf(LOCTEXT("AttributeAlreadyExists", "Attribute identifier provided already exists: {0} {1} ({2}) {3}"),
-			FText::FromName(AttributeIdentifier.GetName()), FText::FromName(AttributeIdentifier.GetBoneName()), FText::AsNumber(AttributeIdentifier.GetBoneIndex()), FText::FromName(AttributeIdentifier.GetType()->GetFName()));
+			ReportErrorf(LOCTEXT("AttributeAlreadyExists", "Attribute identifier with matching name already exists: {0} {1} ({2}) {3}"),
+				FText::FromName(ExistingAttribute->Identifier.GetName()),
+				FText::FromName(ExistingAttribute->Identifier.GetBoneName()),
+				FText::AsNumber(ExistingAttribute->Identifier.GetBoneIndex()),
+				FText::FromName(ExistingAttribute->Identifier.GetType()->GetFName()));
 		}
 	}
 	else
diff --git a/Engine/Source/Developer/AnimationDataController/Private/AnimDataController.cpp b/Engine/Source/Developer/AnimationDataController/Private/AnimDataController.cpp
--- a/Engine/Source/Developer/AnimationDataController/Private/AnimDataController.cpp
+++ b/Engine/Source/Developer/AnimationDataController/Private/AnimDataController.cpp
@@ -1713,12 +1713,12 @@ bool UAnimDataController::AddAttribute(const FAnimationAttributeIdentifier& Attr
 {
 	if (AttributeIdentifier.IsValid())
 	{
-		const bool bAttributeAlreadyExists = Model->AnimatedBoneAttributes.ContainsByPredicate([AttributeIdentifier](const FAnimatedBoneAttribute& Attribute) -> bool
-		{
-			return Attribute.Identifier == AttributeIdentifier;
-		});
+		const FAnimatedBoneAttribute* const ExistingAttribute = Model->AnimatedBoneAttributes.FindByPredicate([AttributeIdentifier](const FAnimatedBoneAttribute& Attribute) -> bool
+			{
+				return Attribute.Identifier.GetName() == AttributeIdentifier.GetName() && Attribute.Identifier.GetBoneName() == AttributeIdentifier.GetBoneName();
+			});
 
-		if (!bAttributeAlreadyExists)
+		if (!ExistingAttribute)
 		{
 			FTransaction Transaction = ConditionalTransaction(LOCTEXT("AddAttribute", "Adding Animated Bone Attribute"), bShouldTransact);
 
@@ -1737,8 +1737,11 @@ bool UAnimDataController::AddAttribute(const FAnimationAttributeIdentifier& Attr
 		}
 		else
 		{
-			ReportErrorf(LOCTEXT("AttributeAlreadyExists", "Attribute identifier provided already exists: {0} {1} ({2}) {3}"),
-			FText::FromName(AttributeIdentifier.GetName()), FText::FromName(AttributeIdentifier.GetBoneName()), FText::AsNumber(AttributeIdentifier.GetBoneIndex()), FText::FromName(AttributeIdentifier.GetType()->GetFName()));
+			ReportErrorf(LOCTEXT("AttributeAlreadyExists", "Attribute identifier with matching name already exists: {0} {1} ({2}) {3}"),
+				FText::FromName(ExistingAttribute->Identifier.GetName()),
+				FText::FromName(ExistingAttribute->Identifier.GetBoneName()),
+				FText::AsNumber(ExistingAttribute->Identifier.GetBoneIndex()),
+				FText::FromName(ExistingAttribute->Identifier.GetType()->GetFName()));
 		}
 	}
 	else
