AnimationSequence cache bucket non-deterministic assets
#fix incorporate number of frames into DataModel GUIDs, in case the animation contains a single pose (uniform keys) this previously would result in identical hashes for two animations with the same single pose but different playback length (unusual but possible)

diff --git a/Engine/Plugins/Animation/AnimationData/Source/AnimationData/Private/AnimSequencerDataModel.cpp b/Engine/Plugins/Animation/AnimationData/Source/AnimationData/Private/AnimSequencerDataModel.cpp
--- a/Engine/Plugins/Animation/AnimationData/Source/AnimationData/Private/AnimSequencerDataModel.cpp
+++ b/Engine/Plugins/Animation/AnimationData/Source/AnimationData/Private/AnimSequencerDataModel.cpp
@@ -627,9 +627,13 @@ UAnimSequence* UAnimationSequencerDataModel::GetAnimationSequence() const
 }
 
 template <typename HasherType>
-void UAnimationSequencerDataModel::GenerateStateHash(HasherType& Hasher) const
+void UAnimationSequencerDataModel::GenerateStateHash(HasherType& Hasher, const FGuidGenerationSettings& InSettings) const
 {
 	Hasher.UpdateString(GetClass()->GetName(), TEXT("C"));
+	if (InSettings.bIncludeTimingData)
+	{
+		Hasher.UpdateData(GetNumberOfFrames(), TEXT("FR"));
+	}
 
 	auto UpdateWithChannel = [&Hasher](const auto& Channel, const TCHAR* Name)
 		{
@@ -648,40 +652,53 @@ void UAnimationSequencerDataModel::GenerateStateHash(HasherType& Hasher) const
 		Hasher.BeginObject(TEXT("RIG"));
 		UpdateWithChannel(RigSection->Weight, TEXT("W"));
 
-		Hasher.BeginObject(TEXT("T"));
-		for (const FTransformParameterNameAndCurves& TransformParameter : RigSection->GetTransformParameterNamesAndCurves())
+		if (InSettings.bIncludeBoneData)
 		{
-			Hasher.BeginObject();
-			Hasher.UpdateString(TransformParameter.ParameterName.ToString(), TEXT("N"));
-			for (int32 Index = 0; Index < 3; ++Index)
+			Hasher.BeginObject(TEXT("T"));
+			for (const FTransformParameterNameAndCurves& TransformParameter : RigSection->GetTransformParameterNamesAndCurves())
 			{
-				Hasher.BeginObject(TEXT("C"));
-				UpdateWithChannel(TransformParameter.Translation[Index], TEXT("T"));
-				UpdateWithChannel(TransformParameter.Rotation[Index], TEXT("R"));
-				UpdateWithChannel(TransformParameter.Scale[Index], TEXT("S"));
+				Hasher.BeginObject();
+				Hasher.UpdateString(TransformParameter.ParameterName.ToString(), TEXT("N"));
+				for (int32 Index = 0; Index < 3; ++Index)
+				{
+					Hasher.BeginObject(TEXT("C"));
+					UpdateWithChannel(TransformParameter.Translation[Index], TEXT("T"));
+					UpdateWithChannel(TransformParameter.Rotation[Index], TEXT("R"));
+					UpdateWithChannel(TransformParameter.Scale[Index], TEXT("S"));
+					Hasher.EndObject();
+				}
 				Hasher.EndObject();
 			}
 			Hasher.EndObject();
 		}
-		Hasher.EndObject();
 
-		Hasher.BeginObject(TEXT("S"));
-		for (const FScalarParameterNameAndCurve& ScalarCurve : RigSection->GetScalarParameterNamesAndCurves())
+		if (InSettings.bIncludeCurveData)
 		{
-			Hasher.BeginObject();
-			Hasher.UpdateString(ScalarCurve.ParameterName.ToString(), TEXT("N"));
-			UpdateWithChannel(ScalarCurve.ParameterCurve, TEXT("C"));
+			Hasher.BeginObject(TEXT("S"));
+			for (const FScalarParameterNameAndCurve& ScalarCurve : RigSection->GetScalarParameterNamesAndCurves())
+			{
+				Hasher.BeginObject();
+				Hasher.UpdateString(ScalarCurve.ParameterName.ToString(), TEXT("N"));
+				UpdateWithChannel(ScalarCurve.ParameterCurve, TEXT("C"));
+				Hasher.EndObject();
+			}
 			Hasher.EndObject();
 		}
 		Hasher.EndObject();
-		Hasher.EndObject();
 	}
 
-	Hasher.UpdateAnimatedBoneAttributes(AnimatedBoneAttributes, TEXT("ABA"));
-	Hasher.UpdateTransformCurves(GetTransformCurves(), TEXT("TC"));
+	if (InSettings.bIncludeAttributeData)
+	{
+		Hasher.UpdateAnimatedBoneAttributes(AnimatedBoneAttributes, TEXT("ABA"));
+	}
+
+	if (InSettings.bIncludeCurveData)
+	{
+		Hasher.UpdateTransformCurves(GetTransformCurves(), TEXT("TC"));
+	}
 }
 
-FGuid UAnimationSequencerDataModel::GenerateGuid() const
+FGuid UAnimationSequencerDataModel::GenerateGuid(const FGuidGenerationSettings& InSettings) const
 {
 	if (CachedRawDataGUID.IsValid())
 	{
@@ -689,7 +706,7 @@ FGuid UAnimationSequencerDataModel::GenerateGuid() const
 	}
 
 	UE::Anim::DataModel::FHasherSha Hasher;
-	GenerateStateHash(Hasher);
+	GenerateStateHash(Hasher, InSettings);
 	return Hasher.FinalGuid();
 }
 
@@ -697,7 +714,7 @@ FGuid UAnimationSequencerDataModel::GenerateGuid() const
 FString UAnimationSequencerDataModel::GenerateDebugStateString() const
 {
 	UE::Anim::DataModel::FHasherCopyToText Hasher;
-	GenerateStateHash(Hasher);
+	GenerateStateHash(Hasher, FGuidGenerationSettings());
 	return Hasher.GetString();
 }
 #endif
diff --git a/Engine/Plugins/Animation/AnimationData/Source/AnimationData/Private/AnimSequencerDataModel.h b/Engine/Plugins/Animation/AnimationData/Source/AnimationData/Private/AnimSequencerDataModel.h
--- a/Engine/Plugins/Animation/AnimationData/Source/AnimationData/Private/AnimSequencerDataModel.h
+++ b/Engine/Plugins/Animation/AnimationData/Source/AnimationData/Private/AnimSequencerDataModel.h
@@ -112,7 +112,7 @@ public:
 	virtual const FAnimatedBoneAttribute* FindAttribute(const FAnimationAttributeIdentifier& AttributeIdentifier) const override;
 	virtual UAnimSequence* GetAnimationSequence() const override;
 	virtual FAnimDataModelModifiedEvent& GetModifiedEvent() override { return ModifiedEvent; }
-	virtual FGuid GenerateGuid() const override;
+	virtual FGuid GenerateGuid(const FGuidGenerationSettings& InSettings) const override;
 #if WITH_EDITOR
 	virtual FString GenerateDebugStateString() const override;
 #endif
@@ -178,7 +178,7 @@ protected:
 	void ClearControlRigData();
 private:	
 	template <typename HasherType>
-	void GenerateStateHash(HasherType& Hasher) const;
+	void GenerateStateHash(HasherType& Hasher, const FGuidGenerationSettings& InSettings) const;
 
 	void InitializeRigHierarchy(UFKControlRig* FKControlRig, USkeleton* Skeleton) const;
 	
diff --git a/Engine/Source/Runtime/Engine/Classes/Animation/AnimData/AnimDataModel.h b/Engine/Source/Runtime/Engine/Classes/Animation/AnimData/AnimDataModel.h
--- a/Engine/Source/Runtime/Engine/Classes/Animation/AnimData/AnimDataModel.h
+++ b/Engine/Source/Runtime/Engine/Classes/Animation/AnimData/AnimDataModel.h
@@ -70,7 +70,7 @@ public:
 	ENGINE_API virtual const FAnimatedBoneAttribute* FindAttribute(const FAnimationAttributeIdentifier& AttributeIdentifier) const override;
 	ENGINE_API virtual UAnimSequence* GetAnimationSequence() const override;
 	virtual FAnimDataModelModifiedEvent& GetModifiedEvent() override { return ModifiedEvent; }
-	ENGINE_API virtual FGuid GenerateGuid() const override;
+	ENGINE_API virtual FGuid GenerateGuid(const FGuidGenerationSettings& InSettings) const override;
 #if WITH_EDITOR
 	ENGINE_API virtual FString GenerateDebugStateString() const override;
 	ENGINE_API virtual void Evaluate(FAnimationPoseData& InOutPoseData, const UE::Anim::DataModel::FEvaluationContext& EvaluationContext) const override;
@@ -102,7 +102,7 @@ private:
 
 	ENGINE_API USkeleton* GetSkeleton() const;
 	template <typename HasherType>
-	void GenerateStateHash(HasherType& Hasher) const;
+	void GenerateStateHash(HasherType& Hasher, const FGuidGenerationSettings& InSettings) const;
 
 private:
 	UPROPERTY(Transient)
diff --git a/Engine/Source/Runtime/Engine/Classes/Animation/AnimData/IAnimationDataModel.h b/Engine/Source/Runtime/Engine/Classes/Animation/AnimData/IAnimationDataModel.h
--- a/Engine/Source/Runtime/Engine/Classes/Animation/AnimData/IAnimationDataModel.h
+++ b/Engine/Source/Runtime/Engine/Classes/Animation/AnimData/IAnimationDataModel.h
@@ -362,10 +362,27 @@ public:
 	*/
 	virtual FAnimDataModelModifiedEvent& GetModifiedEvent() = 0;
 
+	struct FGuidGenerationSettings
+	{
+		FGuidGenerationSettings() : bIncludeBoneData(1), bIncludeCurveData(1), bIncludeAttributeData(1), bIncludeTimingData(1) {}
+
+		uint8 bIncludeBoneData : 1;
+		uint8 bIncludeCurveData : 1;
+		uint8 bIncludeAttributeData : 1;
+		uint8 bIncludeTimingData : 1;
+	};
+
+	/**
+	* @return	GUID representing a configured sub-set of the contained data and state 
+	*
+	* @param	InSettings		Settings used to determine which data should be included when generating the returned GUID
+	*/
+	virtual FGuid GenerateGuid(const FGuidGenerationSettings& InSettings) const = 0;
+
 	/**
 	* @return	GUID representing the contained data and state 
 	*/
-	virtual FGuid GenerateGuid() const = 0;
+	virtual FGuid GenerateGuid() const { return GenerateGuid(FGuidGenerationSettings()); }
 
 #if WITH_EDITOR
 	/**
diff --git a/Engine/Source/Runtime/Engine/Private/Animation/AnimDataModel.cpp b/Engine/Source/Runtime/Engine/Private/Animation/AnimDataModel.cpp
--- a/Engine/Source/Runtime/Engine/Private/Animation/AnimDataModel.cpp
+++ b/Engine/Source/Runtime/Engine/Private/Animation/AnimDataModel.cpp
@@ -507,37 +507,52 @@ const FAnimCurveBase* UAnimDataModel::FindCurve(const FAnimationCurveIdentifier&
 }
 
 template <typename HasherType>
-void UAnimDataModel::GenerateStateHash(HasherType& Hasher) const
+void UAnimDataModel::GenerateStateHash(HasherType& Hasher, const FGuidGenerationSettings& InSettings) const
 {
-	Hasher.BeginObject(TEXT("BAT"));
-	for (const FBoneAnimationTrack& Track : BoneAnimationTracks)
+	if (InSettings.bIncludeTimingData)
+	{		
+		Hasher.UpdateData(GetNumberOfFrames(), TEXT("FR"));
+	}
+
+	if (InSettings.bIncludeBoneData)
 	{
-		Hasher.BeginObject();
-		Hasher.UpdateArray(Track.InternalTrackData.PosKeys, TEXT("P"));
-		Hasher.UpdateArray(Track.InternalTrackData.RotKeys, TEXT("R"));
-		Hasher.UpdateArray(Track.InternalTrackData.ScaleKeys, TEXT("S"));
+		Hasher.BeginObject(TEXT("BAT"));
+		for (const FBoneAnimationTrack& Track : BoneAnimationTracks)
+		{
+			Hasher.BeginObject();
+			Hasher.UpdateArray(Track.InternalTrackData.PosKeys, TEXT("P"));
+			Hasher.UpdateArray(Track.InternalTrackData.RotKeys, TEXT("R"));
+			Hasher.UpdateArray(Track.InternalTrackData.ScaleKeys, TEXT("S"));
+			Hasher.EndObject();
+		}
 		Hasher.EndObject();
 	}
-	Hasher.EndObject();
 
-	Hasher.BeginObject(TEXT("FC"));
-	for (const FFloatCurve& Curve : CurveData.FloatCurves)
+	if (InSettings.bIncludeCurveData)
 	{
-		Hasher.BeginObject();
-		Hasher.UpdateLegacyString(Curve.GetName().ToString(), TEXT("N"));
-		Hasher.UpdateRichCurve(Curve.FloatCurve, TEXT("C"));
+		Hasher.BeginObject(TEXT("FC"));
+		for (const FFloatCurve& Curve : CurveData.FloatCurves)
+		{
+			Hasher.BeginObject();
+			Hasher.UpdateLegacyString(Curve.GetName().ToString(), TEXT("N"));
+			Hasher.UpdateRichCurve(Curve.FloatCurve, TEXT("C"));
+			Hasher.EndObject();
+		}
 		Hasher.EndObject();
+
+		Hasher.UpdateTransformCurves(CurveData.TransformCurves, TEXT("TC"));
 	}
-	Hasher.EndObject();
 
-	Hasher.UpdateTransformCurves(CurveData.TransformCurves, TEXT("TC"));
-	Hasher.UpdateAnimatedBoneAttributes(AnimatedBoneAttributes, TEXT("ABA"));
+	if (InSettings.bIncludeAttributeData)
+	{
+		Hasher.UpdateAnimatedBoneAttributes(AnimatedBoneAttributes, TEXT("ABA"));
+	}
 }
 
-FGuid UAnimDataModel::GenerateGuid() const
+FGuid UAnimDataModel::GenerateGuid(const FGuidGenerationSettings& InSettings) const
 {
 	UE::Anim::DataModel::FHasherSha Hasher;
-	GenerateStateHash(Hasher);
+	GenerateStateHash(Hasher, InSettings);
 	return Hasher.FinalGuid();
 }
 
@@ -545,7 +560,7 @@ FGuid UAnimDataModel::GenerateGuid() const
 FString UAnimDataModel::GenerateDebugStateString() const
 {
 	UE::Anim::DataModel::FHasherCopyToText Hasher;
-	GenerateStateHash(Hasher);
+	GenerateStateHash(Hasher, FGuidGenerationSettings());
 	return Hasher.GetString();
 }
 #endif
diff --git a/Engine/Source/Runtime/Engine/Private/Animation/PoseAsset.cpp b/Engine/Source/Runtime/Engine/Private/Animation/PoseAsset.cpp
--- a/Engine/Source/Runtime/Engine/Private/Animation/PoseAsset.cpp
+++ b/Engine/Source/Runtime/Engine/Private/Animation/PoseAsset.cpp
@@ -943,8 +943,10 @@ void UPoseAsset::PostLoad()
 				}
 			}
 			SourceAnimation->ConditionalPostLoad();
-			
-			if (SourceAnimationRawDataGUID.IsValid() && SourceAnimationRawDataGUID != SourceAnimation->GetDataModel()->GenerateGuid())
+
+			IAnimationDataModel::FGuidGenerationSettings Settings;
+			Settings.bIncludeTimingData = false;
+			if (SourceAnimationRawDataGUID.IsValid() && SourceAnimationRawDataGUID != SourceAnimation->GetDataModel()->GenerateGuid(Settings))
 			{
 				FFormatNamedArguments Args;
 				Args.Add(TEXT("AssetName"), FText::FromString(GetPathName()));
