Ensure that AnimSequencerDataMode does not crash when failing to initialize
#fix Change hard checks over to error logging in case the AnimationSequencerDataModel is not correctly initialized or gets into an invalid state
#misc Added missing skeleton ptr checks

diff --git a/Engine/Plugins/Animation/AnimationData/Source/AnimationData/Private/AnimSequencerController.cpp b/Engine/Plugins/Animation/AnimationData/Source/AnimationData/Private/AnimSequencerController.cpp
--- a/Engine/Plugins/Animation/AnimationData/Source/AnimationData/Private/AnimSequencerController.cpp
+++ b/Engine/Plugins/Animation/AnimationData/Source/AnimationData/Private/AnimSequencerController.cpp
@@ -108,61 +108,63 @@ void UAnimSequencerController::ResizeNumberOfFrames(FFrameNumber NewLength, FFra
 {
 	ValidateModel();
 
-	const TRange<FFrameNumber> PlayRange = Model->GetMovieScene()->GetPlaybackRange();
-	if (NewLength >= 1)
+	if (UMovieScene* MovieScene = Model->GetMovieScene())
 	{
-		if (NewLength != Model->GetNumberOfFrames())
+		const TRange<FFrameNumber> PlayRange = MovieScene->GetPlaybackRange();
+		if (NewLength >= 1)
 		{
-			// Ensure that T0 is within the current play range
-			if (PlayRange.Contains(T0))
+			if (NewLength != Model->GetNumberOfFrames())
 			{
-				// Ensure that the start and end length of either removal or insertion are valid
-				if (T0 < T1)
+				// Ensure that T0 is within the current play range
+				if (PlayRange.Contains(T0))
 				{
-					IAnimationDataModel::FEvaluationAndModificationLock Lock(*ModelInterface);
-					FTransaction Transaction = ConditionalTransaction(LOCTEXT("ResizePlayLength", "Resizing Play Length"), bShouldTransact);
-					
-					const TObjectPtr<UMovieScene>& MovieScene = Model->MovieScene;
-					const FFrameRate CurrentFrameRate = MovieScene->GetDisplayRate();
-					const FFrameNumber CurrentNumberOfFrames = Model->GetNumberOfFrames();
+					// Ensure that the start and end length of either removal or insertion are valid
+					if (T0 < T1)
+					{
+						IAnimationDataModel::FEvaluationAndModificationLock Lock(*ModelInterface);
+						FTransaction Transaction = ConditionalTransaction(LOCTEXT("ResizePlayLength", "Resizing Play Length"), bShouldTransact);
+						
+						const FFrameRate CurrentFrameRate = MovieScene->GetDisplayRate();
+						const FFrameNumber CurrentNumberOfFrames = Model->GetNumberOfFrames();
 
-					ConditionalAction<UE::Anim::FResizePlayLengthInFramesAction>(bShouldTransact, Model.Get(), T0, T1);
-					
-					SetMovieSceneRange(NewLength);
+						ConditionalAction<UE::Anim::FResizePlayLengthInFramesAction>(bShouldTransact, Model.Get(), T0, T1);
+						
+						SetMovieSceneRange(NewLength);
 
-					FSequenceLengthChangedPayload Payload;
+						FSequenceLengthChangedPayload Payload;
 
-					PRAGMA_DISABLE_DEPRECATION_WARNINGS
-					Payload.T0 = static_cast<float>(CurrentFrameRate.AsSeconds(T0));
-					Payload.T1 = static_cast<float>(CurrentFrameRate.AsSeconds(T1));
-					Payload.PreviousLength = static_cast<float>(CurrentFrameRate.AsSeconds(CurrentNumberOfFrames));
-					PRAGMA_ENABLE_DEPRECATION_WARNINGS
+						PRAGMA_DISABLE_DEPRECATION_WARNINGS
+						Payload.T0 = static_cast<float>(CurrentFrameRate.AsSeconds(T0));
+						Payload.T1 = static_cast<float>(CurrentFrameRate.AsSeconds(T1));
+						Payload.PreviousLength = static_cast<float>(CurrentFrameRate.AsSeconds(CurrentNumberOfFrames));
+						PRAGMA_ENABLE_DEPRECATION_WARNINGS
 
-					Payload.PreviousNumberOfFrames = CurrentNumberOfFrames;
-					Payload.Frame0 = T0;
-					Payload.Frame1 = T1;
-					
-					Model->GetNotifier().Notify<FSequenceLengthChangedPayload>(EAnimDataModelNotifyType::SequenceLengthChanged, Payload);
+						Payload.PreviousNumberOfFrames = CurrentNumberOfFrames;
+						Payload.Frame0 = T0;
+						Payload.Frame1 = T1;
+						
+						Model->GetNotifier().Notify<FSequenceLengthChangedPayload>(EAnimDataModelNotifyType::SequenceLengthChanged, Payload);
+					}
+					else
+					{
+						ReportErrorf(LOCTEXT("InvalidEndTimeError", "Invalid T1, smaller that T0 value: T0 {0}, T1 {1}"), FText::AsNumber(T0.Value), FText::AsNumber(T0.Value));
+					}
 				}
 				else
 				{
-					ReportErrorf(LOCTEXT("InvalidEndTimeError", "Invalid T1, smaller that T0 value: T0 {0}, T1 {1}"), FText::AsNumber(T0.Value), FText::AsNumber(T0.Value));
+					ReportErrorf(LOCTEXT("InvalidStartTimeError", "Invalid T0, not within existing play range: T0 {0}, Play Length {1}"), FText::AsNumber(T0.Value), FText::AsNumber(Model->GetPlayLength()));
 				}
 			}
-			else
+			else if (Model->bPopulated)
 			{
-				ReportErrorf(LOCTEXT("InvalidStartTimeError", "Invalid T0, not within existing play range: T0 {0}, Play Length {1}"), FText::AsNumber(T0.Value), FText::AsNumber(Model->GetPlayLength()));
+				Reportf(ELogVerbosity::Display, LOCTEXT("SamePlayLengthWarning", "New play length is same as existing one: {0} frames"), FText::AsNumber(NewLength.Value));
 			}
 		}
-		else if (Model->bPopulated)
+		else
 		{
-			Reportf(ELogVerbosity::Display, LOCTEXT("SamePlayLengthWarning", "New play length is same as existing one: {0} frames"), FText::AsNumber(NewLength.Value));
+			ReportErrorf(LOCTEXT("InvalidPlayLengthError", "Invalid play length value provided: {0} frames"), FText::AsNumber(NewLength.Value));
 		}
-	}
-	else
-	{
-		ReportErrorf(LOCTEXT("InvalidPlayLengthError", "Invalid play length value provided: {0} frames"), FText::AsNumber(NewLength.Value));
-	}
+	}	
 }
 
 void UAnimSequencerController::ResizeInFrames(FFrameNumber NewLength, FFrameNumber T0, FFrameNumber T1, bool bShouldTransact)
@@ -170,46 +172,49 @@ void UAnimSequencerController::ResizeInFrames(FFrameNumber NewLength, FFrameNumb
 	ValidateModel();
 
 	const int32 CurrentNumberOFrames = Model->GetNumberOfFrames();
-	
-	const TRange<FFrameNumber> PlayRange = Model->GetMovieScene()->GetPlaybackRange();
-	if (NewLength >= 1)
+
+	if (UMovieScene* MovieScene = Model->GetMovieScene())
 	{
-		if (NewLength != Model->GetNumberOfFrames())
+		const TRange<FFrameNumber> PlayRange = MovieScene->GetPlaybackRange();
+		if (NewLength >= 1)
 		{
-			// Ensure that T0 is within the current play range
-			if (PlayRange.Contains(T0))
+			if (NewLength != Model->GetNumberOfFrames())
 			{
-				// Ensure that the start and end length of either removal or insertion are valid
-				if (T0 < T1)
+				// Ensure that T0 is within the current play range
+				if (PlayRange.Contains(T0))
 				{
-					IAnimationDataModel::FEvaluationAndModificationLock Lock(*ModelInterface);
-					FBracket Bracket = ConditionalBracket(LOCTEXT("ResizeModel", "Resizing Animation Data"), bShouldTransact);
-					const bool bInserted = NewLength > CurrentNumberOFrames;
-					ResizeNumberOfFrames(NewLength, T0, T1, bShouldTransact);
+					// Ensure that the start and end length of either removal or insertion are valid
+					if (T0 < T1)
+					{
+						IAnimationDataModel::FEvaluationAndModificationLock Lock(*ModelInterface);
+						FBracket Bracket = ConditionalBracket(LOCTEXT("ResizeModel", "Resizing Animation Data"), bShouldTransact);
+						const bool bInserted = NewLength > CurrentNumberOFrames;
+						ResizeNumberOfFrames(NewLength, T0, T1, bShouldTransact);
 
-					const FFrameRate CurrentFrameRate = Model->GetMovieScene()->GetDisplayRate();
-					ResizeCurves(static_cast<float>(CurrentFrameRate.AsSeconds(NewLength)), bInserted, static_cast<float>(CurrentFrameRate.AsSeconds(T0)), static_cast<float>(CurrentFrameRate.AsSeconds(T1)), bShouldTransact);
-					ResizeAttributes(static_cast<float>(CurrentFrameRate.AsSeconds(NewLength)), bInserted, static_cast<float>(CurrentFrameRate.AsSeconds(T0)), static_cast<float>(CurrentFrameRate.AsSeconds(T1)), bShouldTransact);
+						const FFrameRate CurrentFrameRate = MovieScene->GetDisplayRate();
+						ResizeCurves(static_cast<float>(CurrentFrameRate.AsSeconds(NewLength)), bInserted, static_cast<float>(CurrentFrameRate.AsSeconds(T0)), static_cast<float>(CurrentFrameRate.AsSeconds(T1)), bShouldTransact);
+						ResizeAttributes(static_cast<float>(CurrentFrameRate.AsSeconds(NewLength)), bInserted, static_cast<float>(CurrentFrameRate.AsSeconds(T0)), static_cast<float>(CurrentFrameRate.AsSeconds(T1)), bShouldTransact);
+					}
+					else
+					{
+						ReportErrorf(LOCTEXT("InvalidEndTimeError", "Invalid T1, smaller that T0 value: T0 {0}, T1 {1}"), FText::AsNumber(T0.Value), FText::AsNumber(T1.Value));
+					}
 				}
 				else
 				{
-					ReportErrorf(LOCTEXT("InvalidEndTimeError", "Invalid T1, smaller that T0 value: T0 {0}, T1 {1}"), FText::AsNumber(T0.Value), FText::AsNumber(T1.Value));
-				}
+					ReportErrorf(LOCTEXT("InvalidStartTimeError", "Invalid T0, not within existing play range: T0 {0}, Play Length {1}"), FText::AsNumber(T0.Value), FText::AsNumber(CurrentNumberOFrames));
+				}			
 			}
-			else
+			else if (Model->bPopulated)
 			{
-				ReportErrorf(LOCTEXT("InvalidStartTimeError", "Invalid T0, not within existing play range: T0 {0}, Play Length {1}"), FText::AsNumber(T0.Value), FText::AsNumber(CurrentNumberOFrames));
-			}			
+				Reportf(ELogVerbosity::Display,LOCTEXT("SameGetPlayLengthWarning", "New play length is same as existing one: {0} frames"), FText::AsNumber(CurrentNumberOFrames));
+			}
 		}
-		else if (Model->bPopulated)
+		else
 		{
-			Reportf(ELogVerbosity::Display,LOCTEXT("SameGetPlayLengthWarning", "New play length is same as existing one: {0} frames"), FText::AsNumber(CurrentNumberOFrames));
+			ReportErrorf(LOCTEXT("InvalidGetPlayLengthError", "Invalid play length value provided: {0} frames"), FText::AsNumber(CurrentNumberOFrames));
 		}
 	}
-	else
-	{
-		ReportErrorf(LOCTEXT("InvalidGetPlayLengthError", "Invalid play length value provided: {0} frames"), FText::AsNumber(CurrentNumberOFrames));
-	}
 }
 
 void UAnimSequencerController::SetPlayLength(float Length, bool bShouldTransact /*= true*/)
@@ -285,40 +290,41 @@ void UAnimSequencerController::SetFrameRate(FFrameRate FrameRate, bool bShouldTr
 	const float FrameRateInterval = static_cast<float>(FrameRate.AsInterval());
 	if (FrameRate.IsValid() && !FMath::IsNearlyZero(FrameRateInterval) && FrameRateInterval > 0.f)
 	{
-		const TObjectPtr<UMovieScene>& MovieScene = Model->MovieScene;
-		
-		// Need to verify frame-rate
-		const FFrameRate CurrentFrameRate = MovieScene->GetDisplayRate();
-
-		if (FrameRate.IsMultipleOf(CurrentFrameRate) || FrameRate.IsFactorOf(CurrentFrameRate) || !Model->bPopulated)
+		if (UMovieScene* MovieScene = Model->GetMovieScene())
 		{
-			IAnimationDataModel::FEvaluationAndModificationLock Lock(*ModelInterface);
-			FTransaction Transaction = ConditionalTransaction(LOCTEXT("SetFrameRate", "Setting Frame Rate"), bShouldTransact);
-			ConditionalAction<UE::Anim::FSetFrameRateAction>(bShouldTransact, Model.Get());
+			// Need to verify frame-rate
+			const FFrameRate CurrentFrameRate = MovieScene->GetDisplayRate();
 
-			const FFrameNumber CurrentNumberOfFrames = Model->GetNumberOfFrames();
-		
-			MovieScene->SetDisplayRate(FrameRate);
-			MovieScene->SetTickResolutionDirectly(FrameRate);
+			if (FrameRate.IsMultipleOf(CurrentFrameRate) || FrameRate.IsFactorOf(CurrentFrameRate) || !Model->bPopulated)
+			{
+				IAnimationDataModel::FEvaluationAndModificationLock Lock(*ModelInterface);
+				FTransaction Transaction = ConditionalTransaction(LOCTEXT("SetFrameRate", "Setting Frame Rate"), bShouldTransact);
+				ConditionalAction<UE::Anim::FSetFrameRateAction>(bShouldTransact, Model.Get());
 
-			const FFrameTime TransformedFrameTime = FFrameRate::TransformTime(CurrentNumberOfFrames, CurrentFrameRate, FrameRate);
-			if (Model->bPopulated && !FMath::IsNearlyZero(TransformedFrameTime.GetSubFrame(), KINDA_SMALL_NUMBER))
-			{ 
-				static const FNumberFormattingOptions DurationFormatOptions = FNumberFormattingOptions()
-					.SetMinimumFractionalDigits(8)
-					.SetMaximumFractionalDigits(8);
-				ReportWarningf(LOCTEXT("TransformedFrametimePrecisionWarning", "Insufficient precision while converting between frame rates {0} and {1}: previous number of frames {2}, converted frame-time {3}"), CurrentFrameRate.ToPrettyText(), FrameRate.ToPrettyText(), FText::AsNumber(CurrentNumberOfFrames.Value), FText::AsNumber(TransformedFrameTime.AsDecimal(), &DurationFormatOptions));
-			}
+				const FFrameNumber CurrentNumberOfFrames = Model->GetNumberOfFrames();
+		
+				MovieScene->SetDisplayRate(FrameRate);
+				MovieScene->SetTickResolutionDirectly(FrameRate);
+
+				const FFrameTime TransformedFrameTime = FFrameRate::TransformTime(CurrentNumberOfFrames, CurrentFrameRate, FrameRate);
+				if (Model->bPopulated && !FMath::IsNearlyZero(TransformedFrameTime.GetSubFrame(), KINDA_SMALL_NUMBER))
+				{ 
+					static const FNumberFormattingOptions DurationFormatOptions = FNumberFormattingOptions()
+						.SetMinimumFractionalDigits(8)
+						.SetMaximumFractionalDigits(8);
+					ReportWarningf(LOCTEXT("TransformedFrametimePrecisionWarning", "Insufficient precision while converting between frame rates {0} and {1}: previous number of frames {2}, converted frame-time {3}"), CurrentFrameRate.ToPrettyText(), FrameRate.ToPrettyText(), FText::AsNumber(CurrentNumberOfFrames.Value), FText::AsNumber(TransformedFrameTime.AsDecimal(), &DurationFormatOptions));
+				}
 
-			SetMovieSceneRange(TransformedFrameTime.GetFrame());
+				SetMovieSceneRange(TransformedFrameTime.GetFrame());
 
-			FFrameRateChangedPayload Payload;
-			Payload.PreviousFrameRate = CurrentFrameRate;
-			Model->GetNotifier().Notify(EAnimDataModelNotifyType::FrameRateChanged, Payload);
-		}
-		else
-		{
-			ReportErrorf(LOCTEXT("NonCompatibleFrameRateError", "Incompatible frame rate provided: {0} not a multiple or factor of {1}"), FrameRate.ToPrettyText(), CurrentFrameRate.ToPrettyText());
+				FFrameRateChangedPayload Payload;
+				Payload.PreviousFrameRate = CurrentFrameRate;
+				Model->GetNotifier().Notify(EAnimDataModelNotifyType::FrameRateChanged, Payload);
+			}
+			else
+			{
+				ReportErrorf(LOCTEXT("NonCompatibleFrameRateError", "Incompatible frame rate provided: {0} not a multiple or factor of {1}"), FrameRate.ToPrettyText(), CurrentFrameRate.ToPrettyText());
+			}
 		}
 	}	
 	else
diff --git a/Engine/Plugins/Animation/AnimationData/Source/AnimationData/Private/AnimSequencerDataModel.cpp b/Engine/Plugins/Animation/AnimationData/Source/AnimationData/Private/AnimSequencerDataModel.cpp
--- a/Engine/Plugins/Animation/AnimationData/Source/AnimationData/Private/AnimSequencerDataModel.cpp
+++ b/Engine/Plugins/Animation/AnimationData/Source/AnimationData/Private/AnimSequencerDataModel.cpp
@@ -119,33 +119,44 @@ USkeleton* UAnimationSequencerDataModel::GetSkeleton() const
 	return Skeleton;
 }
 
-void UAnimationSequencerDataModel::InitializeRigHierarchy(UFKControlRig* FKControlRig, USkeleton* Skeleton) const
+void UAnimationSequencerDataModel::InitializeRigHierarchy(UFKControlRig* FKControlRig, const USkeleton* Skeleton) const
 {
-	FScopeLock Lock(&EvaluationLock);
-	
-	UFKControlRig::FRigElementInitializationOptions InitOptions;
-	InitOptions.bImportCurves = false;	
-	if(UMovieSceneControlRigParameterSection* Section = GetFKControlRigSection())
+	if (FKControlRig && Skeleton)
 	{
-		for (const FScalarParameterNameAndCurve& AnimCurve : Section->GetScalarParameterNamesAndCurves())
+		FScopeLock Lock(&EvaluationLock);
+	
+		UFKControlRig::FRigElementInitializationOptions InitOptions;
+		InitOptions.bImportCurves = false;	
+		if(UMovieSceneControlRigParameterSection* Section = GetFKControlRigSection())
 		{
-			InitOptions.CurveNames.Add(UFKControlRig::GetControlTargetName(AnimCurve.ParameterName, ERigElementType::Curve));
-		}
+			for (const FScalarParameterNameAndCurve& AnimCurve : Section->GetScalarParameterNamesAndCurves())
+			{
+				InitOptions.CurveNames.Add(UFKControlRig::GetControlTargetName(AnimCurve.ParameterName, ERigElementType::Curve));
+			}
 
-		for (const FTransformParameterNameAndCurves& BoneCurve : Section->GetTransformParameterNamesAndCurves())
-		{
-			InitOptions.BoneNames.Add(UFKControlRig::GetControlTargetName(BoneCurve.ParameterName, ERigElementType::Bone));
+			for (const FTransformParameterNameAndCurves& BoneCurve : Section->GetTransformParameterNamesAndCurves())
+			{
+				InitOptions.BoneNames.Add(UFKControlRig::GetControlTargetName(BoneCurve.ParameterName, ERigElementType::Bone));
+			}
 		}
-	}
-	InitOptions.bGenerateBoneControls = InitOptions.BoneNames.Num() > 0;
-	FKControlRig->SetInitializationOptions(InitOptions);
+		InitOptions.bGenerateBoneControls = InitOptions.BoneNames.Num() > 0;
+		FKControlRig->SetInitializationOptions(InitOptions);
 
-	FKControlRig->Initialize();
+		FKControlRig->Initialize();
 
-	FKControlRig->SetApplyMode(UseDirectFKControlRigMode == 1 ? EControlRigFKRigExecuteMode::Direct : EControlRigFKRigExecuteMode::Replace);
-	FKControlRig->SetBoneInitialTransformsFromRefSkeleton(Skeleton->GetReferenceSkeleton());
-	FKControlRig->Evaluate_AnyThread();
-	bRigHierarchyInitialized = true;
+		FKControlRig->SetApplyMode(UseDirectFKControlRigMode == 1 ? EControlRigFKRigExecuteMode::Direct : EControlRigFKRigExecuteMode::Replace);
+		FKControlRig->SetBoneInitialTransformsFromRefSkeleton(Skeleton->GetReferenceSkeleton());
+		FKControlRig->Evaluate_AnyThread();
+		bRigHierarchyInitialized = true;
+	}
+	else if (!FKControlRig)
+	{
+		IAnimationDataController::ReportObjectErrorf(this, LOCTEXT("FailedToInitHierarchyFKCR", "Unable to initialize RigHierarchy due to invalid FK ControlRig for ({0})"), FText::FromString(GetPathName()));
+	}
+	else if (!Skeleton)
+	{
+		IAnimationDataController::ReportObjectErrorf(this, LOCTEXT("FailedToInitHierarchySkeleton", "Unable to initialize RigHierarchy due to missing USkeleton for ({0})"), FText::FromString(GetPathName()));
+	}	
 }
 
 void UAnimationSequencerDataModel::InitializeFKControlRig(UFKControlRig* FKControlRig, USkeleton* Skeleton, bool bForceHierarchyInitialization /*=false*/) const
@@ -283,22 +294,32 @@ void UAnimationSequencerDataModel::PostEditUndo()
 double UAnimationSequencerDataModel::GetPlayLength() const
 {
 	ValidateSequencerData();
-	return MovieScene->GetDisplayRate().AsSeconds(GetNumberOfFrames());
+	if (MovieScene)
+	{
+		return MovieScene->GetDisplayRate().AsSeconds(GetNumberOfFrames());
+	}
+	
+	return 0.f;
 }
 
 int32 UAnimationSequencerDataModel::GetNumberOfFrames() const
 {	
 	ValidateSequencerData();
-	const TRange<FFrameNumber> FrameRange = MovieScene->GetPlaybackRange();	
-	const TRangeBound<FFrameNumber>& UpperRange = FrameRange.GetUpperBound();
-	const bool bInclusive = UpperRange.IsInclusive();
-	int32 Value = UpperRange.GetValue().Value;
-	if (!bInclusive)
-	{		
-		Value = FMath::Max(Value - 1, 1);
-	}
+	if (MovieScene)
+	{
+		const TRange<FFrameNumber> FrameRange = MovieScene->GetPlaybackRange();	
+		const TRangeBound<FFrameNumber>& UpperRange = FrameRange.GetUpperBound();
+		const bool bInclusive = UpperRange.IsInclusive();
+		int32 Value = UpperRange.GetValue().Value;
+		if (!bInclusive)
+		{		
+			Value = FMath::Max(Value - 1, 1);
+		}
 
-	return Value;
+		return Value;
+	}
+	
+	return 0;
 }
 
 int32 UAnimationSequencerDataModel::GetNumberOfKeys() const
@@ -309,7 +330,12 @@ int32 UAnimationSequencerDataModel::GetNumberOfKeys() const
 FFrameRate UAnimationSequencerDataModel::GetFrameRate() const
 {	
 	ValidateSequencerData();
-	return MovieScene->GetDisplayRate();
+	if (MovieScene)
+	{
+		return MovieScene->GetDisplayRate();
+	}
+
+	return UAnimationSettings::Get()->GetDefaultFrameRate();
 }
 
 const TArray<FBoneAnimationTrack>& UAnimationSequencerDataModel::GetBoneAnimationTracks() const
@@ -862,7 +888,7 @@ void UAnimationSequencerDataModel::OnNotify(const EAnimDataModelNotifyType& Noti
 
 UMovieSceneControlRigParameterTrack* UAnimationSequencerDataModel::GetControlRigTrack() const
 {
-	return MovieScene->FindTrack<UMovieSceneControlRigParameterTrack>();
+	return MovieScene ? MovieScene->FindTrack<UMovieSceneControlRigParameterTrack>() : nullptr;
 }
 
 UMovieSceneControlRigParameterSection* UAnimationSequencerDataModel::GetFKControlRigSection() const
@@ -1023,64 +1049,119 @@ void UAnimationSequencerDataModel::ValidateData() const
 
 void UAnimationSequencerDataModel::ValidateSequencerData() const
 {
-	checkf(MovieScene, TEXT("%s: No Movie Scene found for SequencerDataModel"), *GetPathName());
-
-	const int32 NumberOfTracks = MovieScene->GetTracks().Num();
-	checkf(NumberOfTracks == 1, TEXT("%s: Invalid number of Tracks in Movie Scene expected 1 but found %i"), *GetPathName(), NumberOfTracks);
-		
-	const UMovieSceneControlRigParameterTrack* Track = MovieScene->FindTrack<UMovieSceneControlRigParameterTrack>();
-	checkf(Track, TEXT("%s: Unable to find Control Rig Track"), *GetPathName());
-
-	const int32 NumberOfSections = Track->GetAllSections().Num();
-	checkf(NumberOfSections == 1, TEXT("%s: Invalid number of Sections found for Control Rig Track expected 1 but found %i"), *GetPathName(), NumberOfSections);
-
-	const UMovieSceneControlRigParameterSection* Section = GetFKControlRigSection();
-	checkf(Section, TEXT("%s: Unable to find Control Rig Section"), *GetPathName());
+	if (MovieScene)
+	{
+		const int32 NumberOfTracks = MovieScene->GetTracks().Num();
+		if (NumberOfTracks == 1)
+		{
+			if(const UMovieSceneControlRigParameterTrack* Track = MovieScene->FindTrack<UMovieSceneControlRigParameterTrack>())
+			{
+				const int32 NumberOfSections = Track->GetAllSections().Num();
+				if (NumberOfSections == 1)
+				{
+					if(const UMovieSceneControlRigParameterSection* Section = GetFKControlRigSection())
+					{
+						if (const UControlRig* ControlRig = Section->GetControlRig())
+						{
+							if (!ControlRig->IsA<UFKControlRig>())
+							{
+								IAnimationDataController::ReportMessage(this, FText::Format(LOCTEXT("InvalidControlRigClass", "Unexpected class {0} on ControlRig, expecting FKControlRig"), FText::FromString(ControlRig->GetClass()->GetPathName())), ELogVerbosity::Error);
+							}
+						}
+						else
+						{
+							IAnimationDataController::ReportMessage(this, LOCTEXT("MissingControlRig", "Unable to find Control Rig"), ELogVerbosity::Error);
+						}
+					}
+					else
+					{
+						IAnimationDataController::ReportMessage(this, LOCTEXT("MissingControlRigSection", "Unable to find Control Rig Section"), ELogVerbosity::Error);	
+					}
+				}
+				else
+				{
+					IAnimationDataController::ReportMessage(this, FText::Format(LOCTEXT("InvalidNumberOfTracks", "Invalid number of Sections found for Control Rig Track expected 1 but found {0}"), FText::AsNumber(NumberOfSections)), ELogVerbosity::Error);
+				}
+			}
+			else
+			{
+				IAnimationDataController::ReportMessage(this, LOCTEXT("MissingControlRigTrack", "Unable to find Control Rig Track"), ELogVerbosity::Error);	
+			}
+		}
+		else
+		{
+			IAnimationDataController::ReportMessage(this, FText::Format(LOCTEXT("InvalidNumberOfTracks", "Invalid number of Tracks in Movie Scene expected 1 but found {0}"), FText::AsNumber(NumberOfTracks)), ELogVerbosity::Error);
+		}
+	}
+	else
+	{
+		IAnimationDataController::ReportMessage(this, LOCTEXT("MissingMovieScene", "No Movie Scene found for SequencerDataModel"), ELogVerbosity::Error);
+	}
 }
 
 void UAnimationSequencerDataModel::ValidateControlRigData() const
 {
-	const UMovieSceneControlRigParameterSection* Section = GetFKControlRigSection();
-	checkf(Section, TEXT("Unable to find Control Rig Section"));
-
-	UControlRig* ControlRig = Section->GetControlRig();
-	checkf(ControlRig, TEXT("Unable to find Control Rig instance for Section"));
-
-	checkf(ControlRig->IsA<UFKControlRig>(), TEXT("Invalid class for Control Rig expected UFKControlRig"));
-
-	const URigHierarchy* Hierarchy = ControlRig->GetHierarchy();
-	if (Hierarchy && ValidationMode)
+	if(const UMovieSceneControlRigParameterSection* Section = GetFKControlRigSection())
 	{
-		// Validate Rig Hierarchy against the outer Animation Sequence its (reference) Skeleton		
-		if (const USkeleton* Skeleton = GetSkeleton())
+		if (UControlRig* ControlRig = Section->GetControlRig())
 		{
-			const FReferenceSkeleton& ReferenceSkeleton = Skeleton->GetReferenceSkeleton();
-			const int32 NumberOfBones = ReferenceSkeleton.GetNum();
-
-			// Validating the bone elements against the reference skeleton bones
-			for (int32 BoneIndex = 0; BoneIndex < NumberOfBones; ++BoneIndex)
+			if (ControlRig->IsA<UFKControlRig>())
 			{
-				const FName ExpectedBoneName = ReferenceSkeleton.GetBoneName(BoneIndex);
-				const bool bIsVirtualBone = ExpectedBoneName.ToString().StartsWith(VirtualBoneNameHelpers::VirtualBonePrefix);
-				if (!bIsVirtualBone)
+				if (bRigHierarchyInitialized)
 				{
-					const FRigElementKey BoneKey(ExpectedBoneName, ERigElementType::Bone);
-					const FRigBoneElement* BoneElement = Hierarchy->Find<FRigBoneElement>(BoneKey);
-					checkf(BoneElement, TEXT("Unable to find FRigBoneElement in RigHierarchy for Bone with name: %s"), *ExpectedBoneName.ToString());
-		
-					const int32 ParentBoneIndex = ReferenceSkeleton.GetParentIndex(BoneIndex);
-					if (BoneElement && ParentBoneIndex != INDEX_NONE)
+					const URigHierarchy* Hierarchy = ControlRig->GetHierarchy();
+					if (Hierarchy && ValidationMode)
+					{
+						// Validate Rig Hierarchy against the outer Animation Sequence its (reference) Skeleton		
+						if (const USkeleton* Skeleton = GetSkeleton())
+						{
+							const FReferenceSkeleton& ReferenceSkeleton = Skeleton->GetReferenceSkeleton();
+							const int32 NumberOfBones = ReferenceSkeleton.GetNum();
+
+							// Validating the bone elements against the reference skeleton bones
+							for (int32 BoneIndex = 0; BoneIndex < NumberOfBones; ++BoneIndex)
+							{
+								const FName ExpectedBoneName = ReferenceSkeleton.GetBoneName(BoneIndex);
+								const bool bIsVirtualBone = ExpectedBoneName.ToString().StartsWith(VirtualBoneNameHelpers::VirtualBonePrefix);
+								if (!bIsVirtualBone)
+								{
+									const FRigElementKey BoneKey(ExpectedBoneName, ERigElementType::Bone);
+									const FRigBoneElement* BoneElement = Hierarchy->Find<FRigBoneElement>(BoneKey);
+									checkf(BoneElement, TEXT("Unable to find FRigBoneElement in RigHierarchy for Bone with name: %s"), *ExpectedBoneName.ToString());
+
+									const int32 ParentBoneIndex = ReferenceSkeleton.GetParentIndex(BoneIndex);
+									if (BoneElement && ParentBoneIndex != INDEX_NONE)
+									{
+										const FName ExpectedParentBoneName = ReferenceSkeleton.GetBoneName(ParentBoneIndex);
+										const FRigElementKey ParentBoneKey(ExpectedParentBoneName, ERigElementType::Bone);
+
+										const FRigBoneElement* ParentBoneElement = Hierarchy->Find<FRigBoneElement>(ParentBoneKey);
+										checkf(BoneElement->ParentElement == ParentBoneElement, TEXT("Unexpected Parent Element for Bone %s. Expected %s but found %s"), *ExpectedBoneName.ToString(), *ExpectedParentBoneName.ToString(), *ParentBoneElement->GetDisplayName().ToString());
+									}	
+								}
+							}
+						}
+					}
+					else if (!Hierarchy)
 					{
-						const FName ExpectedParentBoneName = ReferenceSkeleton.GetBoneName(ParentBoneIndex);
-						const FRigElementKey ParentBoneKey(ExpectedParentBoneName, ERigElementType::Bone);
-            
-						const FRigBoneElement* ParentBoneElement = Hierarchy->Find<FRigBoneElement>(ParentBoneKey);
-						checkf(BoneElement->ParentElement == ParentBoneElement, TEXT("Unexpected Parent Element for Bone %s. Expected %s but found %s"), *ExpectedBoneName.ToString(), *ExpectedParentBoneName.ToString(), *ParentBoneElement->GetDisplayName().ToString());
-					}	
+						IAnimationDataController::ReportMessage(this, LOCTEXT("MissingHierarchy", "Unable to retrieve Control Rig Hierarchy"), ELogVerbosity::Error);	
+					}
 				}
 			}
+			else
+			{
+				IAnimationDataController::ReportMessage(this, FText::Format(LOCTEXT("InvalidControlRigClass", "Unexpected class {0} on ControlRig, expecting FKControlRig"), FText::FromString(ControlRig->GetClass()->GetPathName())), ELogVerbosity::Error);	
+			}
+		}
+		else
+		{
+			IAnimationDataController::ReportMessage(this, LOCTEXT("MissingControlRig", "Unable to find Control Rig"), ELogVerbosity::Error);	
 		}
 	}
+	else
+	{
+		IAnimationDataController::ReportMessage(this, LOCTEXT("MissingControlRigSection", "Unable to find Control Rig Section"), ELogVerbosity::Error);	
+	}
 }
 
 void UAnimationSequencerDataModel::ValidateLegacyAgainstControlRigData() const
@@ -1151,56 +1232,55 @@ void UAnimationSequencerDataModel::IterateTransformControlCurve(const FName& Bon
 	ValidateSequencerData();
 	ValidateControlRigData();
 
-	const UMovieSceneControlRigParameterSection* Section = GetFKControlRigSection();
-	UControlRig* ControlRig = Section->GetControlRig();
-	const URigHierarchy* Hierarchy = ControlRig->GetHierarchy();
-	
-	const FRigElementKey BoneControlKey(UFKControlRig::GetControlName(BoneName, ERigElementType::Bone), ERigElementType::Control);
-	if (const FTransformParameterNameAndCurves* ControlCurvePtr = Section->GetTransformParameterNamesAndCurves().FindByPredicate([CurveName = BoneControlKey.Name](const FTransformParameterNameAndCurves& TransformParameter)
-	{
-		return TransformParameter.ParameterName == CurveName;
-	}))
+	if (const UMovieSceneControlRigParameterSection* Section = GetFKControlRigSection())
 	{
-		const FTransformParameterNameAndCurves& ControlCurve = *ControlCurvePtr;
+		const FRigElementKey BoneControlKey(UFKControlRig::GetControlName(BoneName, ERigElementType::Bone), ERigElementType::Control);
+		if (const FTransformParameterNameAndCurves* ControlCurvePtr = Section->GetTransformParameterNamesAndCurves().FindByPredicate([CurveName = BoneControlKey.Name](const FTransformParameterNameAndCurves& TransformParameter)
+		{
+			return TransformParameter.ParameterName == CurveName;
+		}))
+		{
+			const FTransformParameterNameAndCurves& ControlCurve = *ControlCurvePtr;
 
-		FVector3f Location(0.f);
-		FVector3f EulerAngles(0.f);
-		FVector3f Scale = FVector3f::OneVector;
+			FVector3f Location(0.f);
+			FVector3f EulerAngles(0.f);
+			FVector3f Scale = FVector3f::OneVector;
 
-		// Check whether or not any data is contained
-		bool bContainsData = false;
-		bool bContainsKeys = false;
-		for (int32 ChannelIndex = 0; ChannelIndex < 3; ++ChannelIndex)
-		{
-			bContainsData |= ControlCurve.Translation[ChannelIndex].HasAnyData();
-			bContainsKeys |= ControlCurve.Translation[ChannelIndex].GetNumKeys() != 0;
-			bContainsData |= ControlCurve.Rotation[ChannelIndex].HasAnyData();		
-			bContainsKeys |= ControlCurve.Rotation[ChannelIndex].GetNumKeys() != 0;		
-			bContainsData |= ControlCurve.Scale[ChannelIndex].HasAnyData();
-			bContainsKeys |= ControlCurve.Scale[ChannelIndex].GetNumKeys() != 0;
-		}
+			// Check whether or not any data is contained
+			bool bContainsData = false;
+			bool bContainsKeys = false;
+			for (int32 ChannelIndex = 0; ChannelIndex < 3; ++ChannelIndex)
+			{
+				bContainsData |= ControlCurve.Translation[ChannelIndex].HasAnyData();
+				bContainsKeys |= ControlCurve.Translation[ChannelIndex].GetNumKeys() != 0;
+				bContainsData |= ControlCurve.Rotation[ChannelIndex].HasAnyData();		
+				bContainsKeys |= ControlCurve.Rotation[ChannelIndex].GetNumKeys() != 0;		
+				bContainsData |= ControlCurve.Scale[ChannelIndex].HasAnyData();
+				bContainsKeys |= ControlCurve.Scale[ChannelIndex].GetNumKeys() != 0;
+			}
 
-		if (bContainsData)
-		{
-			const int32 NumberOfKeysToIterate = InFrameNumbers != nullptr ? InFrameNumbers->Num() : GetNumberOfKeys();
-			for (int32 KeyIndex = 0; KeyIndex < NumberOfKeysToIterate; ++KeyIndex)
+			if (bContainsData)
 			{
-				const FFrameNumber Frame = InFrameNumbers != nullptr ? (*InFrameNumbers)[KeyIndex] : FFrameNumber(KeyIndex);
-				for (int32 ChannelIndex = 0; ChannelIndex < 3; ++ChannelIndex)
+				const int32 NumberOfKeysToIterate = InFrameNumbers != nullptr ? InFrameNumbers->Num() : GetNumberOfKeys();
+				for (int32 KeyIndex = 0; KeyIndex < NumberOfKeysToIterate; ++KeyIndex)
 				{
-					ControlCurve.Translation[ChannelIndex].Evaluate(Frame, Location[ChannelIndex]);
-					ControlCurve.Rotation[ChannelIndex].Evaluate(Frame, EulerAngles[ChannelIndex]);
-					ControlCurve.Scale[ChannelIndex].Evaluate(Frame, Scale[ChannelIndex]);
-				}
+					const FFrameNumber Frame = InFrameNumbers != nullptr ? (*InFrameNumbers)[KeyIndex] : FFrameNumber(KeyIndex);
+					for (int32 ChannelIndex = 0; ChannelIndex < 3; ++ChannelIndex)
+					{
+						ControlCurve.Translation[ChannelIndex].Evaluate(Frame, Location[ChannelIndex]);
+						ControlCurve.Rotation[ChannelIndex].Evaluate(Frame, EulerAngles[ChannelIndex]);
+						ControlCurve.Scale[ChannelIndex].Evaluate(Frame, Scale[ChannelIndex]);
+					}
 
-				FTransform Transform;
-				Transform.SetLocation(FVector(Location));
-				Transform.SetRotation(FQuat::MakeFromEuler(FVector(EulerAngles)));
-				Transform.SetScale3D(FVector(Scale));
+					FTransform Transform;
+					Transform.SetLocation(FVector(Location));
+					Transform.SetRotation(FQuat::MakeFromEuler(FVector(EulerAngles)));
+					Transform.SetScale3D(FVector(Scale));
 
-				Transform.NormalizeRotation();
+					Transform.NormalizeRotation();
 
-				IterationFunction(Transform, Frame);
+					IterationFunction(Transform, Frame);
+				}
 			}
 		}
 	}
@@ -1429,7 +1509,11 @@ void UAnimationSequencerDataModel::EvaluateTrack(UMovieSceneControlRigParameterT
 	if (SectionsInRange.Num())
 	{
 		const UMovieSceneControlRigParameterSection* FKRigSection = CastChecked<UMovieSceneControlRigParameterSection>(SectionsInRange[0]);
-		checkf(FKRigSection->ControlRigClass->GetDefaultObject()->IsA<UFKControlRig>(), TEXT("Unexpected class %s on ControlRig, expecting FKControlRig"), *FKRigSection->ControlRigClass->GetPathName());
+		if (!FKRigSection->ControlRigClass->GetDefaultObject()->IsA<UFKControlRig>())
+		{
+			IAnimationDataController::ReportMessage(this, FText::Format(LOCTEXT("InvalidControlRigClass", "Unexpected class {0} on ControlRig, expecting FKControlRig"), FText::FromString(FKRigSection->ControlRigClass->GetPathName())), ELogVerbosity::Error);
+			return;
+		}
 		
 		bool bWasDoNotKey = false;
 		bWasDoNotKey = FKRigSection->GetDoNotKey();
@@ -1681,119 +1765,121 @@ void UAnimationSequencerDataModel::IterateBoneKeys(const FName& BoneName, TFunct
 	ValidateSequencerData();
 	ValidateControlRigData();
 
-	const UMovieSceneControlRigParameterSection* Section = GetFKControlRigSection();
-	UControlRig* ControlRig = Section->GetControlRig();
-	const URigHierarchy* Hierarchy = ControlRig->GetHierarchy();
-	
-	const FRigElementKey BoneControlKey(UFKControlRig::GetControlName(BoneName, ERigElementType::Bone), ERigElementType::Control);
-	if (const FTransformParameterNameAndCurves* ControlCurvePtr = Section->GetTransformParameterNamesAndCurves().FindByPredicate([CurveName = BoneControlKey.Name](const FTransformParameterNameAndCurves& TransformParameter)
-	{
-		return TransformParameter.ParameterName == CurveName;
-	}))
+	if (const UMovieSceneControlRigParameterSection* Section = GetFKControlRigSection())
 	{
-		const FTransformParameterNameAndCurves& ControlCurve = *ControlCurvePtr;
-
-		struct FChannelInfo
-		{
-			bool bConstant = true;
-			bool bUniform = true;
-		};			
-		FChannelInfo PosChannels[3];
-		FChannelInfo RotChannels[3];
-		FChannelInfo ScaleChannels[3];
-
-		const int32 NumberOfKeys = GetNumberOfKeys();
-
-		int32 MaxNumberOfKeys = 1;
-		TSet<FFrameNumber> FrameNumbers;
-		for (int32 ChannelIndex = 0; ChannelIndex < 3; ++ChannelIndex)
-		{
-			auto ProcessChannel = [&MaxNumberOfKeys, NumberOfKeys](const FMovieSceneFloatChannel& Channel, FChannelInfo& Info)
-			{
-				if (Channel.HasAnyData())
-				{
-					if (Channel.GetNumKeys() == 0)
-					{
-						Info.bUniform = false;
-					}
-					else
-					{
-						Info.bConstant = false;
-						if (Channel.GetNumKeys() != NumberOfKeys)
-						{
-							Info.bUniform = false;
-						}
-
-						MaxNumberOfKeys = FMath::Max(MaxNumberOfKeys, Channel.GetNumKeys());
-					}
-				}
-			};
-
-			ProcessChannel(ControlCurve.Translation[ChannelIndex], PosChannels[ChannelIndex]);
-			ProcessChannel(ControlCurve.Rotation[ChannelIndex], RotChannels[ChannelIndex]);
-			ProcessChannel(ControlCurve.Scale[ChannelIndex], ScaleChannels[ChannelIndex]);
-		}
-
-
-		const int32 NumberOfKeysToIterate = MaxNumberOfKeys;
-		const FRigControlElement* BoneControl = Hierarchy->Find<FRigControlElement>(BoneControlKey);
-		const FTransform InitialTransform = BoneControl ? BoneControl->GetTransform().Initial.Local.Get() : FTransform::Identity;
-
-		// Initialize components with initial control values (in case there is no default value nor keys)
-		FVector3f PreviousPos = FVector3f(InitialTransform.GetLocation());
-		FVector3f PreviousScale = FVector3f(InitialTransform.GetScale3D());
-		FVector3f PreviousRot = FVector3f(InitialTransform.GetRotation().Euler());
-		FQuat4f PreviousQuat;
-
-		for (int32 KeyIndex = 0; KeyIndex < NumberOfKeysToIterate; ++KeyIndex)
-		{
-			const FFrameNumber Frame = KeyIndex;
-			
-			for (int32 ChannelIndex = 0; ChannelIndex < 3; ++ChannelIndex)
-			{
-				if (PosChannels[ChannelIndex].bConstant)
-				{
-					if (KeyIndex == 0)
-					{
-						PreviousPos[ChannelIndex] = ControlCurve.Translation[ChannelIndex].GetDefault().Get(PreviousPos[ChannelIndex]);
-					}
-				}
-				else
-				{
-					ControlCurve.Translation[ChannelIndex].Evaluate(Frame, PreviousPos[ChannelIndex]);
-				}
-				
-				if (RotChannels[ChannelIndex].bConstant)
-				{
-					if (KeyIndex == 0)
-					{
-						PreviousRot[ChannelIndex] = ControlCurve.Rotation[ChannelIndex].GetDefault().Get(PreviousRot[ChannelIndex]);
-					}
-				}
-				else
-				{
-					ControlCurve.Rotation[ChannelIndex].Evaluate(Frame, PreviousRot[ChannelIndex]);
-				}
-
-				if (ScaleChannels[ChannelIndex].bConstant)
-				{
-					if (KeyIndex == 0)
-					{
-						PreviousScale[ChannelIndex] = ControlCurve.Scale[ChannelIndex].GetDefault().Get(PreviousScale[ChannelIndex]);
-					}
-				}
-				else
-				{
-					ControlCurve.Scale[ChannelIndex].Evaluate(Frame, PreviousScale[ChannelIndex]);
-				}
-			}
-
-			PreviousQuat = FQuat4f::MakeFromEuler(PreviousRot);
-			if(!IterationFunction(PreviousPos, PreviousQuat, PreviousScale, Frame))
-			{
-				return;
-			}
-		}
+	    UControlRig* ControlRig = Section->GetControlRig();
+	    const URigHierarchy* Hierarchy = ControlRig->GetHierarchy();
+	    
+	    const FRigElementKey BoneControlKey(UFKControlRig::GetControlName(BoneName, ERigElementType::Bone), ERigElementType::Control);
+	    if (const FTransformParameterNameAndCurves* ControlCurvePtr = Section->GetTransformParameterNamesAndCurves().FindByPredicate([CurveName = BoneControlKey.Name](const FTransformParameterNameAndCurves& TransformParameter)
+	    {
+		    return TransformParameter.ParameterName == CurveName;
+	    }))
+	    {
+		    const FTransformParameterNameAndCurves& ControlCurve = *ControlCurvePtr;
+
+		    struct FChannelInfo
+		    {
+			    bool bConstant = true;
+			    bool bUniform = true;
+		    };			
+		    FChannelInfo PosChannels[3];
+		    FChannelInfo RotChannels[3];
+		    FChannelInfo ScaleChannels[3];
+
+		    const int32 NumberOfKeys = GetNumberOfKeys();
+
+		    int32 MaxNumberOfKeys = 1;
+		    TSet<FFrameNumber> FrameNumbers;
+		    for (int32 ChannelIndex = 0; ChannelIndex < 3; ++ChannelIndex)
+		    {
+			    auto ProcessChannel = [&MaxNumberOfKeys, NumberOfKeys](const FMovieSceneFloatChannel& Channel, FChannelInfo& Info)
+			    {
+				    if (Channel.HasAnyData())
+				    {
+					    if (Channel.GetNumKeys() == 0)
+					    {
+						    Info.bUniform = false;
+					    }
+					    else
+					    {
+						    Info.bConstant = false;
+						    if (Channel.GetNumKeys() != NumberOfKeys)
+						    {
+							    Info.bUniform = false;
+						    }
+
+						    MaxNumberOfKeys = FMath::Max(MaxNumberOfKeys, Channel.GetNumKeys());
+					    }
+				    }
+			    };
+
+			    ProcessChannel(ControlCurve.Translation[ChannelIndex], PosChannels[ChannelIndex]);
+			    ProcessChannel(ControlCurve.Rotation[ChannelIndex], RotChannels[ChannelIndex]);
+			    ProcessChannel(ControlCurve.Scale[ChannelIndex], ScaleChannels[ChannelIndex]);
+		    }
+
+
+		    const int32 NumberOfKeysToIterate = MaxNumberOfKeys;
+		    const FRigControlElement* BoneControl = Hierarchy->Find<FRigControlElement>(BoneControlKey);
+		    const FTransform InitialTransform = BoneControl ? BoneControl->GetTransform().Initial.Local.Get() : FTransform::Identity;
+
+		    // Initialize components with initial control values (in case there is no default value nor keys)
+		    FVector3f PreviousPos = FVector3f(InitialTransform.GetLocation());
+		    FVector3f PreviousScale = FVector3f(InitialTransform.GetScale3D());
+		    FVector3f PreviousRot = FVector3f(InitialTransform.GetRotation().Euler());
+		    FQuat4f PreviousQuat;
+
+		    for (int32 KeyIndex = 0; KeyIndex < NumberOfKeysToIterate; ++KeyIndex)
+		    {
+			    const FFrameNumber Frame = KeyIndex;
+			    
+			    for (int32 ChannelIndex = 0; ChannelIndex < 3; ++ChannelIndex)
+			    {
+				    if (PosChannels[ChannelIndex].bConstant)
+				    {
+					    if (KeyIndex == 0)
+					    {
+						    PreviousPos[ChannelIndex] = ControlCurve.Translation[ChannelIndex].GetDefault().Get(PreviousPos[ChannelIndex]);
+					    }
+				    }
+				    else
+				    {
+					    ControlCurve.Translation[ChannelIndex].Evaluate(Frame, PreviousPos[ChannelIndex]);
+				    }
+				    
+				    if (RotChannels[ChannelIndex].bConstant)
+				    {
+					    if (KeyIndex == 0)
+					    {
+						    PreviousRot[ChannelIndex] = ControlCurve.Rotation[ChannelIndex].GetDefault().Get(PreviousRot[ChannelIndex]);
+					    }
+				    }
+				    else
+				    {
+					    ControlCurve.Rotation[ChannelIndex].Evaluate(Frame, PreviousRot[ChannelIndex]);
+				    }
+
+				    if (ScaleChannels[ChannelIndex].bConstant)
+				    {
+					    if (KeyIndex == 0)
+					    {
+						    PreviousScale[ChannelIndex] = ControlCurve.Scale[ChannelIndex].GetDefault().Get(PreviousScale[ChannelIndex]);
+					    }
+				    }
+				    else
+				    {
+					    ControlCurve.Scale[ChannelIndex].Evaluate(Frame, PreviousScale[ChannelIndex]);
+				    }
+			    }
+
+			    PreviousQuat = FQuat4f::MakeFromEuler(PreviousRot);
+			    if(!IterationFunction(PreviousPos, PreviousQuat, PreviousScale, Frame))
+			    {
+				    return;
+			    }
+		    }
+	    }
 	}
 }
 
diff --git a/Engine/Plugins/Animation/AnimationData/Source/AnimationData/Private/AnimSequencerDataModel.h b/Engine/Plugins/Animation/AnimationData/Source/AnimationData/Private/AnimSequencerDataModel.h
--- a/Engine/Plugins/Animation/AnimationData/Source/AnimationData/Private/AnimSequencerDataModel.h
+++ b/Engine/Plugins/Animation/AnimationData/Source/AnimationData/Private/AnimSequencerDataModel.h
@@ -180,7 +180,7 @@ private:
 	template <typename HasherType>
 	void GenerateStateHash(HasherType& Hasher, const FGuidGenerationSettings& InSettings) const;
 
-	void InitializeRigHierarchy(UFKControlRig* FKControlRig, USkeleton* Skeleton) const;
+	void InitializeRigHierarchy(UFKControlRig* FKControlRig, const USkeleton* Skeleton) const;
 	
 	/** Dynamic delegate event allows scripting to register to any broadcast-ed notify. */
 	UPROPERTY(BlueprintAssignable, Transient, Category = AnimationDataModel, meta = (ScriptName = "ModifiedEvent", AllowPrivateAccess = "true"))
