Make sure we do not do any redundant work during UAnimSequencerController::UpdateWithSkeleton in case nothing is _actually_ going to change

diff --git a/Engine/Plugins/Animation/AnimationData/Source/AnimationData/Private/AnimSequencerController.cpp b/Engine/Plugins/Animation/AnimationData/Source/AnimationData/Private/AnimSequencerController.cpp
--- a/Engine/Plugins/Animation/AnimationData/Source/AnimationData/Private/AnimSequencerController.cpp
+++ b/Engine/Plugins/Animation/AnimationData/Source/AnimationData/Private/AnimSequencerController.cpp
@@ -2257,18 +2257,59 @@ void UAnimSequencerController::UpdateWithSkeleton(USkeleton* TargetSkeleton, boo
 	{
 		OpenBracket(LOCTEXT("SettingNewskeleton", "Updating Skeleton for Animation Data Model"), bShouldTransact);
 		{
-			// (re-)generate the rig hierarchy
-			Model->InitializeFKControlRig(CastChecked<UFKControlRig>(Model->GetControlRig()), TargetSkeleton);
-
-			// Remove any bone tracks that do not exist in the new hierarchy
-			RemoveBoneTracksMissingFromSkeleton(TargetSkeleton, bShouldTransact);
-
-			// Remove any control/curves which were created in InitializeFKControlRig which are no longer keyed
-			RemoveUnusedControlsAndCurves();		
+			// Work out if any work has to be performed.
+			TArray<FName> BonesToBeRemoved;
+			TArray<FName> CurvesToBeRemoved;
+
+			if(TargetSkeleton)
+			{
+			    const FReferenceSkeleton& ReferenceSkeleton = TargetSkeleton->GetReferenceSkeleton();
+			    IAnimationDataModel::FEvaluationAndModificationLock Lock(*ModelInterface);
+			    if (UMovieSceneControlRigParameterSection* Section = Model->GetFKControlRigSection())
+			    {
+				    const TArray<FTransformParameterNameAndCurves>& BoneTrackCurves = Section->GetTransformParameterNamesAndCurves();
+				    for (const FTransformParameterNameAndCurves& BoneCurve : BoneTrackCurves)
+				    {
+					    const FName ExpectedBoneName = UFKControlRig::GetControlTargetName(BoneCurve.ParameterName, ERigElementType::Bone);
+					    
+					    if (ReferenceSkeleton.FindBoneIndex(ExpectedBoneName) == INDEX_NONE)
+					    {
+						    // Bone no longer exists on skeleton
+						    BonesToBeRemoved.Add(ExpectedBoneName);
+					    }
+				    }
+
+				    TArray<FName> SkeletonCurveNames;
+				    TargetSkeleton->GetCurveMetaDataNames(SkeletonCurveNames);
+
+				    const TArray<FScalarParameterNameAndCurve>& SectionCurves = Section->GetScalarParameterNamesAndCurves();
+				    for (const FScalarParameterNameAndCurve& SectionCurve : SectionCurves)
+				    {
+					    const FName ExpectedCurveName = UFKControlRig::GetControlTargetName(SectionCurve.ParameterName, ERigElementType::Curve);
+					    if (!SkeletonCurveNames.Contains(ExpectedCurveName))
+					    {
+						    // Curve no longer exists on skeleton
+						    CurvesToBeRemoved.Add(ExpectedCurveName);
+					    }
+				    }
+			    }
+			}
+
+			if (BonesToBeRemoved.Num() || CurvesToBeRemoved.Num())
+			{
+				// (re-)generate the rig hierarchy
+				Model->InitializeFKControlRig(CastChecked<UFKControlRig>(Model->GetControlRig()), TargetSkeleton, true);
+				
+				// Remove any bone tracks that do not exist in the new hierarchy
+				RemoveBoneTracksMissingFromSkeleton(TargetSkeleton, bShouldTransact);
 
-			// Forcefully re-generate legacy data structures
-			Model->RegenerateLegacyCurveData();
+				// Remove any control/curves which were created in InitializeFKControlRig which are no longer keyed
+				RemoveUnusedControlsAndCurves();
 
+				// Forcefully re-generate legacy data structures
+				Model->RegenerateLegacyCurveData();
+			}
+			
 			// Notify of skeleton change
 			Model->GetNotifier().Notify(EAnimDataModelNotifyType::SkeletonChanged);
 		}
