//UE5/Main - Build Engine Localization - LogLocTextHelper: Warning: Text conflict from LOCTEXT macro for namespace 'AnimSequencerDataModel' and key 'InvalidNumberOfTracks'
#fix swap LOCTEXT identifier

diff --git a/Engine/Plugins/Animation/AnimationData/Source/AnimationData/Private/AnimSequencerDataModel.cpp b/Engine/Plugins/Animation/AnimationData/Source/AnimationData/Private/AnimSequencerDataModel.cpp
--- a/Engine/Plugins/Animation/AnimationData/Source/AnimationData/Private/AnimSequencerDataModel.cpp
+++ b/Engine/Plugins/Animation/AnimationData/Source/AnimationData/Private/AnimSequencerDataModel.cpp
@@ -1080,7 +1080,7 @@ void UAnimationSequencerDataModel::ValidateSequencerData() const
 				}
 				else
 				{
-					IAnimationDataController::ReportMessage(this, FText::Format(LOCTEXT("InvalidNumberOfTracks", "Invalid number of Sections found for Control Rig Track expected 1 but found {0}"), FText::AsNumber(NumberOfSections)), ELogVerbosity::Error);
+					IAnimationDataController::ReportMessage(this, FText::Format(LOCTEXT("InvalidNumberOfSections", "Invalid number of Sections found for Control Rig Track expected 1 but found {0}"), FText::AsNumber(NumberOfSections)), ELogVerbosity::Error);
 				}
 			}
 			else
