[IK Retarget] Fixes regresion in animation exported from retargeter.

diff --git a/Engine/Plugins/Animation/AnimationData/Source/AnimationData/Private/AnimSequencerController.cpp b/Engine/Plugins/Animation/AnimationData/Source/AnimationData/Private/AnimSequencerController.cpp
--- a/Engine/Plugins/Animation/AnimationData/Source/AnimationData/Private/AnimSequencerController.cpp
+++ b/Engine/Plugins/Animation/AnimationData/Source/AnimationData/Private/AnimSequencerController.cpp
@@ -2258,6 +2258,7 @@ void UAnimSequencerController::UpdateWithSkeleton(USkeleton* TargetSkeleton, boo
 		OpenBracket(LOCTEXT("SettingNewskeleton", "Updating Skeleton for Animation Data Model"), bShouldTransact);
 		{
 			// Work out if any work has to be performed.
+			TArray<FName> BonesToBeAdded;
 			TArray<FName> BonesToBeRemoved;
 			TArray<FName> CurvesToBeRemoved;
 
@@ -2267,6 +2268,7 @@ void UAnimSequencerController::UpdateWithSkeleton(USkeleton* TargetSkeleton, boo
 			    IAnimationDataModel::FEvaluationAndModificationLock Lock(*ModelInterface);
 			    if (UMovieSceneControlRigParameterSection* Section = Model->GetFKControlRigSection())
 			    {
+			    	// find bones needing to be REMOVED
 				    const TArray<FTransformParameterNameAndCurves>& BoneTrackCurves = Section->GetTransformParameterNamesAndCurves();
 				    for (const FTransformParameterNameAndCurves& BoneCurve : BoneTrackCurves)
 				    {
@@ -2279,9 +2281,28 @@ void UAnimSequencerController::UpdateWithSkeleton(USkeleton* TargetSkeleton, boo
 					    }
 				    }
 
+			    	// find bones needing to be ADDED
+				    const TArray<FMeshBoneInfo>& RawRefBones = ReferenceSkeleton.GetRawRefBoneInfo();
+			    	for (const FMeshBoneInfo& RawRefBone : RawRefBones)
+			    	{
+			    		bool bFoundBoneTrack = false;
+			    		for (const FTransformParameterNameAndCurves& BoneCurve : BoneTrackCurves)
+			    		{
+			    			if (BoneCurve.ParameterName == RawRefBone.Name)
+			    			{
+			    				bFoundBoneTrack = true;
+			    				break;
+			    			}
+			    		}
+			    		if (!bFoundBoneTrack)
+			    		{
+			    			BonesToBeAdded.Add(RawRefBone.Name);
+			    		}
+			    	}
+
+			    	// find curves needing to be REMOVED
 				    TArray<FName> SkeletonCurveNames;
 				    TargetSkeleton->GetCurveMetaDataNames(SkeletonCurveNames);
-
 				    const TArray<FScalarParameterNameAndCurve>& SectionCurves = Section->GetScalarParameterNamesAndCurves();
 				    for (const FScalarParameterNameAndCurve& SectionCurve : SectionCurves)
 				    {
@@ -2295,7 +2316,7 @@ void UAnimSequencerController::UpdateWithSkeleton(USkeleton* TargetSkeleton, boo
 			    }
 			}
 
-			if (BonesToBeRemoved.Num() || CurvesToBeRemoved.Num())
+			if (BonesToBeAdded.Num() || BonesToBeRemoved.Num() || CurvesToBeRemoved.Num())
 			{
 				// (re-)generate the rig hierarchy
 				Model->InitializeFKControlRig(CastChecked<UFKControlRig>(Model->GetControlRig()), TargetSkeleton, true);
