[CrashReport][Assert] Engine!FRichCurve::GetKey() [RichCurve.cpp]
#fix setting curve keys with time values non-aligning with frames can cause diverging key.time values - mitigate the undo-behaviour by using the actual set time-value rather than the user requested timing

diff --git a/Engine/Plugins/Animation/AnimationData/Source/AnimationData/Private/AnimSequencerController.cpp b/Engine/Plugins/Animation/AnimationData/Source/AnimationData/Private/AnimSequencerController.cpp
--- a/Engine/Plugins/Animation/AnimationData/Source/AnimationData/Private/AnimSequencerController.cpp
+++ b/Engine/Plugins/Animation/AnimationData/Source/AnimationData/Private/AnimSequencerController.cpp
@@ -1188,12 +1188,16 @@ bool UAnimSequencerController::SetCurveKey(const FAnimationCurveIdentifier& Curv
 		FCurveChangedPayload Payload;
 		Payload.Identifier = CurveId;
 
-		auto SetControlCurve = [this, CurveId, Key](bool bUpdate)
+		auto SetControlCurve = [this, CurveId, Key](bool bUpdate, float& OutKeyTiming)
 		{
-			if(CurveId.CurveType == ERawCurveTrackTypes::RCT_Float && !SetCurveControlKey(CurveId.CurveName, Key, bUpdate))
+			// FMovieSceneFloatChannel stores keys along frames rather than time-based keying, this means Key.Time can be different than actual key set on the float-channel (max 240k fps). So adjust the timing using OutKeyTiming to match for undo/redo and rich key mirroring.
+			if(CurveId.CurveType == ERawCurveTrackTypes::RCT_Float && !SetCurveControlKey(CurveId.CurveName, Key, bUpdate, OutKeyTiming))
 			{
-				ReportError(LOCTEXT("FailedtoSetCurveControlKey", "Failed to set curve control key"));
-			}	
+				ReportError(LOCTEXT("FailedToSetCurveControlKey", "Failed to set curve control key"));
+				return false;
+			}
+
+			return true;
 		};
 
 		// Set or add rich curve value
@@ -1204,26 +1208,29 @@ bool UAnimSequencerController::SetCurveKey(const FAnimationCurveIdentifier& Curv
 			FTransaction Transaction = ConditionalTransaction(LOCTEXT("SetNamedCurveKey", "Setting Curve Key"), bShouldTransact);
 			// Cache old value for action
 			const FRichCurveKey CurrentKey = RichCurve->GetKey(Handle);
-			ConditionalAction<UE::Anim::FSetRichCurveKeyAction>(bShouldTransact,  CurveId, CurrentKey);
 
-			SetControlCurve(true);
-
-			// Set the new value
-			RichCurve->SetKeyValue(Handle, Key.Value);
-
-			Model->GetNotifier().Notify(EAnimDataModelNotifyType::CurveChanged, Payload);
+			float MovieSceneKeyTiming = 0.f;
+			if (SetControlCurve(true, MovieSceneKeyTiming))
+			{
+				ConditionalAction<UE::Anim::FSetRichCurveKeyAction>(bShouldTransact,  CurveId, CurrentKey);
+				// Set the new value
+				RichCurve->SetKeyValue(Handle, Key.Value);
+				Model->GetNotifier().Notify(EAnimDataModelNotifyType::CurveChanged, Payload);
+			}
 		}
 		else
 		{
 			FTransaction Transaction = ConditionalTransaction(LOCTEXT("AddNamedCurveKey", "Adding Curve Key"), bShouldTransact);
-			ConditionalAction<UE::Anim::FRemoveRichCurveKeyAction>(bShouldTransact,  CurveId, Key.Time);
 
-			SetControlCurve(false);
-
-			// Add the new key
-			RichCurve->AddKey(Key.Time, Key.Value);
+			float MovieSceneKeyTiming = 0.f;
+			if (SetControlCurve(false, MovieSceneKeyTiming))
+			{
+				ConditionalAction<UE::Anim::FRemoveRichCurveKeyAction>(bShouldTransact,  CurveId, MovieSceneKeyTiming);
+				// Add the new key
+				RichCurve->AddKey(MovieSceneKeyTiming, Key.Value);
 
-			Model->GetNotifier().Notify(EAnimDataModelNotifyType::CurveChanged, Payload);
+				Model->GetNotifier().Notify(EAnimDataModelNotifyType::CurveChanged, Payload);
+			}
 		}
 
 		return true;
@@ -3465,7 +3472,7 @@ bool UAnimSequencerController::SetCurveControlKeys(const FName& CurveName, const
 	return false;
 }
 
-bool UAnimSequencerController::SetCurveControlKey(const FName& CurveName, const FRichCurveKey& Key, bool bUpdateKey) const
+bool UAnimSequencerController::SetCurveControlKey(const FName& CurveName, const FRichCurveKey& Key, bool bUpdateKey, float& OutKeyTiming) const
 {
 	if (UMovieSceneControlRigParameterSection* Section = Model->GetFKControlRigSection())
 	{
@@ -3521,6 +3528,8 @@ bool UAnimSequencerController::SetCurveControlKey(const FName& CurveName, const
 					const FFrameNumber FrameNumber = ParameterCurvePair->ParameterCurve.GetTickResolution().AsFrameTime(Key.Time).RoundToFrame();
 					check(bUpdateKey || ParameterCurvePair->ParameterCurve.GetData().FindKey(FrameNumber) == INDEX_NONE);
 					const FKeyHandle KeyHandle = ParameterCurvePair->ParameterCurve.GetData().UpdateOrAddKey(FrameNumber, MovieSceneValue);
+					
+					OutKeyTiming = ParameterCurvePair->ParameterCurve.GetTickResolution().AsSeconds(FrameNumber);
 					return KeyHandle != FKeyHandle::Invalid();
 				}
 				else
diff --git a/Engine/Plugins/Animation/AnimationData/Source/AnimationData/Private/AnimSequencerController.h b/Engine/Plugins/Animation/AnimationData/Source/AnimationData/Private/AnimSequencerController.h
--- a/Engine/Plugins/Animation/AnimationData/Source/AnimationData/Private/AnimSequencerController.h
+++ b/Engine/Plugins/Animation/AnimationData/Source/AnimationData/Private/AnimSequencerController.h
@@ -112,7 +112,7 @@ private:
 	bool RenameCurveControl(const FName& CurveName, const FName& NewCurveName) const;
 	bool RemoveCurveControl(const FName& CurveName) const;
 	bool SetCurveControlKeys(const FName& CurveName, const TArray<FRichCurveKey>& CurveKeys) const;
-	bool SetCurveControlKey(const FName& CurveName, const FRichCurveKey& Key, bool bUpdateKey) const;
+	bool SetCurveControlKey(const FName& CurveName, const FRichCurveKey& Key, bool bUpdateKey, float& OutKeyTiming) const;
 	bool RemoveCurveControlKey(const FName& CurveName, float Time) const;
 	bool DuplicateCurveControl(const FName& CurveName, const FName& DuplicateCurveName) const;
 
