[CrashReport][Assert] MovieScene!TMovieSceneCurveChannelImpl::GetInterpolationForTime() [MovieSceneCurveChannelImpl.cpp]
#fix add missing evaluation/modification scope locks

diff --git a/Engine/Plugins/Animation/AnimationData/Source/AnimationData/Private/AnimSequencerController.cpp b/Engine/Plugins/Animation/AnimationData/Source/AnimationData/Private/AnimSequencerController.cpp
--- a/Engine/Plugins/Animation/AnimationData/Source/AnimationData/Private/AnimSequencerController.cpp
+++ b/Engine/Plugins/Animation/AnimationData/Source/AnimationData/Private/AnimSequencerController.cpp
@@ -1220,6 +1220,7 @@ bool UAnimSequencerController::SetCurveKey(const FAnimationCurveIdentifier& Curv
 		}
 		else
 		{
+			IAnimationDataModel::FEvaluationAndModificationLock Lock(*ModelInterface);
 			FTransaction Transaction = ConditionalTransaction(LOCTEXT("AddNamedCurveKey", "Adding Curve Key"), bShouldTransact);
 
 			float MovieSceneKeyTiming = 0.f;
@@ -1312,6 +1313,7 @@ bool UAnimSequencerController::SetCurveKeys(const FAnimationCurveIdentifier& Cur
 
 		if(CurveId.CurveType == ERawCurveTrackTypes::RCT_Float)
 		{
+			IAnimationDataModel::FEvaluationAndModificationLock Lock(*ModelInterface);
 			if (SetCurveControlKeys(CurveId.CurveName, CurveKeys))
 			{
 				bKeysSet = true;
@@ -2251,6 +2253,7 @@ void UAnimSequencerController::RemoveUnusedControlsAndCurves() const
 					return true;
 				});
 
+				IAnimationDataModel::FEvaluationAndModificationLock Lock(*ModelInterface);
 				for (const FRigElementKey& Key : ElementsToRemove)
 				{
 					Hierarchy->GetController()->RemoveElement(Key);
