Fix-up incorrect return values from CL 46304680 (bug-hawk)

diff --git a/Engine/Plugins/Animation/AnimationData/Source/AnimationData/Private/AnimSequencerController.cpp b/Engine/Plugins/Animation/AnimationData/Source/AnimationData/Private/AnimSequencerController.cpp
--- a/Engine/Plugins/Animation/AnimationData/Source/AnimationData/Private/AnimSequencerController.cpp
+++ b/Engine/Plugins/Animation/AnimationData/Source/AnimationData/Private/AnimSequencerController.cpp
@@ -1216,6 +1216,8 @@ bool UAnimSequencerController::SetCurveKey(const FAnimationCurveIdentifier& Curv
 				// Set the new value
 				RichCurve->SetKeyValue(Handle, Key.Value);
 				Model->GetNotifier().Notify(EAnimDataModelNotifyType::CurveChanged, Payload);
+				
+				return true;
 			}
 		}
 		else
@@ -1231,10 +1233,10 @@ bool UAnimSequencerController::SetCurveKey(const FAnimationCurveIdentifier& Curv
 				RichCurve->AddKey(MovieSceneKeyTiming, Key.Value);
 
 				Model->GetNotifier().Notify(EAnimDataModelNotifyType::CurveChanged, Payload);
+				
+				return true;
 			}
 		}
-
-		return true;
 	}
 	else
 	{
