//UE5/Release-5.7 - UE.EditorAutomation(RunTest=Animation) Multiple Platforms - Multiple Tests failing due to unexpected and unknown time entry errors
#fix initialize key time value with incoming parameter, in case of Transform curve keys this won't be populated by SetControlCurve lambda
#misc fixed overlapping test names and added some additional failure logging

diff --git a/Engine/Plugins/Animation/AnimationData/Source/AnimationData/Private/AnimSequencerController.cpp b/Engine/Plugins/Animation/AnimationData/Source/AnimationData/Private/AnimSequencerController.cpp
--- a/Engine/Plugins/Animation/AnimationData/Source/AnimationData/Private/AnimSequencerController.cpp
+++ b/Engine/Plugins/Animation/AnimationData/Source/AnimationData/Private/AnimSequencerController.cpp
@@ -1209,7 +1209,7 @@ bool UAnimSequencerController::SetCurveKey(const FAnimationCurveIdentifier& Curv
 			// Cache old value for action
 			const FRichCurveKey CurrentKey = RichCurve->GetKey(Handle);
 
-			float MovieSceneKeyTiming = 0.f;
+			float MovieSceneKeyTiming = Key.Time;
 			if (SetControlCurve(true, MovieSceneKeyTiming))
 			{
 				ConditionalAction<UE::Anim::FSetRichCurveKeyAction>(bShouldTransact,  CurveId, CurrentKey);
@@ -1225,7 +1225,7 @@ bool UAnimSequencerController::SetCurveKey(const FAnimationCurveIdentifier& Curv
 			IAnimationDataModel::FEvaluationAndModificationLock Lock(*ModelInterface);
 			FTransaction Transaction = ConditionalTransaction(LOCTEXT("AddNamedCurveKey", "Adding Curve Key"), bShouldTransact);
 
-			float MovieSceneKeyTiming = 0.f;
+			float MovieSceneKeyTiming = Key.Time;
 			if (SetControlCurve(false, MovieSceneKeyTiming))
 			{
 				ConditionalAction<UE::Anim::FRemoveRichCurveKeyAction>(bShouldTransact,  CurveId, MovieSceneKeyTiming);
