Animation Data: do not check channel proxies while evaluating the track (side effect of 47440003)
- that change fixes the baking process but also highlights that reconstructing/querying the channels proxies is not thread safe (the anim compression is async for example)

#rnx

diff --git a/Engine/Plugins/Animation/AnimationData/Source/AnimationData/Private/AnimSequencerDataModel.cpp b/Engine/Plugins/Animation/AnimationData/Source/AnimationData/Private/AnimSequencerDataModel.cpp
--- a/Engine/Plugins/Animation/AnimationData/Source/AnimationData/Private/AnimSequencerDataModel.cpp
+++ b/Engine/Plugins/Animation/AnimationData/Source/AnimationData/Private/AnimSequencerDataModel.cpp
@@ -1498,12 +1498,28 @@ void UAnimationSequencerDataModel::EvaluateTrack(UMovieSceneControlRigParameterT
 	const FFrameTime InterpolationTime = EvaluationContext.InterpolationType == EAnimInterpolationType::Step ? EvaluationContext.SampleTime.FloorToFrame() : EvaluationContext.SampleTime;
 	const FFrameTime BoneSampleTime = FFrameRate::TransformTime(InterpolationTime, EvaluationContext.SampleFrameRate, MovieScene->GetTickResolution());	
 
-	// Retrieve section withing range of requested evaluation frame 
-	const TArray<UMovieSceneSection*, TInlineAllocator<4>> SectionsInRange = CR_Track->FindAllSections(BoneSampleTime.FrameNumber);
-	//ensureMsgf(SectionsInRange.Num() == 1, TEXT("Unable to retrieve section within range of request evaluation frame %i for %s"), BoneSampleTime.FrameNumber.Value, *GetAnimationSequence()->GetPathName());
-	if (SectionsInRange.Num())
+	auto FindFirstActiveSectionAtTime = [CR_Track, Time = BoneSampleTime.FrameNumber]() -> UMovieSceneSection*
 	{
-		const UMovieSceneControlRigParameterSection* FKRigSection = CastChecked<UMovieSceneControlRigParameterSection>(SectionsInRange[0]);
+		if (CR_Track && !CR_Track->IsEvalDisabled())
+		{
+			for (UMovieSceneSection* Section : CR_Track->GetAllSections())
+			{
+				if (Section && Section->IsActive() && !CR_Track->IsRowEvalDisabled(Section->GetRowIndex()))
+				{
+					if (Section->GetRange().Contains(Time))
+					{
+						return Section;
+					}
+				}
+			}
+		}
+		return nullptr;
+	};
+	
+	// Retrieve section within range of requested evaluation frame
+	if (UMovieSceneSection* Section = FindFirstActiveSectionAtTime())
+	{
+		const UMovieSceneControlRigParameterSection* FKRigSection = CastChecked<UMovieSceneControlRigParameterSection>(Section);
 		if (!FKRigSection->ControlRigClass->GetDefaultObject()->IsA<UFKControlRig>())
 		{
 			IAnimationDataController::ReportMessage(this, FText::Format(LOCTEXT("InvalidControlRigClass", "Unexpected class {0} on ControlRig, expecting FKControlRig"), FText::FromString(FKRigSection->ControlRigClass->GetPathName())), ELogVerbosity::Error);
