Switched AnimSequencerDataModel critical section to recursive (non-shared) mutex (simplifies UAnimationSequencerDataModel::EvaluateBoneTrackTransform impl)
Added missing evaluation locks to BoneTrack related APIs

#Synced-CL 51091569

diff --git a/Engine/Plugins/Animation/AnimationData/Source/AnimationData/Private/AnimSequencerDataModel.cpp b/Engine/Plugins/Animation/AnimationData/Source/AnimationData/Private/AnimSequencerDataModel.cpp
--- a/Engine/Plugins/Animation/AnimationData/Source/AnimationData/Private/AnimSequencerDataModel.cpp
+++ b/Engine/Plugins/Animation/AnimationData/Source/AnimationData/Private/AnimSequencerDataModel.cpp
@@ -125,8 +125,8 @@ void UAnimationSequencerDataModel::InitializeRigHierarchy(UFKControlRig* FKContr
 {
 	if (FKControlRig && Skeleton)
 	{
-		FScopeLock Lock(&EvaluationLock);
-	
+		UE::TUniqueLock<UE::FRecursiveMutex> Lock(EvaluationLock);
+
 		UFKControlRig::FRigElementInitializationOptions InitOptions;
 		InitOptions.bImportCurves = false;	
 		if(UMovieSceneControlRigParameterSection* Section = GetFKControlRigSection())
@@ -378,6 +378,7 @@ bool UAnimationSequencerDataModel::IsValidBoneTrackIndex(int32 TrackIndex) const
 
 int32 UAnimationSequencerDataModel::GetNumBoneTracks() const
 {
+	UE::TUniqueLock<UE::FRecursiveMutex> Lock(EvaluationLock);
 	ValidateSequencerData();
 		
 	if(const UMovieSceneControlRigParameterSection* Section = GetFKControlRigSection())
@@ -390,6 +391,8 @@ int32 UAnimationSequencerDataModel::GetNumBoneTracks() const
 
 void UAnimationSequencerDataModel::GetBoneTrackNames(TArray<FName>& OutNames) const
 {
+	UE::TUniqueLock<UE::FRecursiveMutex> Lock(EvaluationLock);
+
 	if(const UMovieSceneControlRigParameterSection* Section = GetFKControlRigSection())
 	{
 		for (const FTransformParameterNameAndCurves& TransformParameter : Section->GetTransformParameterNamesAndCurves())
@@ -503,6 +506,7 @@ const FRichCurve* UAnimationSequencerDataModel::FindRichCurve(const FAnimationCu
 
 bool UAnimationSequencerDataModel::IsValidBoneTrackName(const FName& TrackName) const
 {
+	UE::TUniqueLock<UE::FRecursiveMutex> Lock(EvaluationLock);
 	ValidateSequencerData();
 
 	if (const UMovieSceneControlRigParameterSection* Section = GetFKControlRigSection())
@@ -516,19 +520,25 @@ bool UAnimationSequencerDataModel::IsValidBoneTrackName(const FName& TrackName)
 
 FTransform UAnimationSequencerDataModel::GetBoneTrackTransform(FName TrackName, const FFrameNumber& FrameNumber) const
 {
+	UE::TUniqueLock<UE::FRecursiveMutex> Lock(EvaluationLock);
+
 	const TArray<FFrameNumber> FrameNumbers = { FrameNumber };
-	TArray<FTransform> Transforms;	
+	TArray<FTransform> Transforms;
 	GenerateTransformKeysForControl(TrackName, FrameNumbers, Transforms);
 	return Transforms.Num() ? Transforms[0] : FTransform::Identity;
 }
 
 void UAnimationSequencerDataModel::GetBoneTrackTransforms(FName TrackName, const TArray<FFrameNumber>& FrameNumbers, TArray<FTransform>& OutTransforms) const
 {
+	UE::TUniqueLock<UE::FRecursiveMutex> Lock(EvaluationLock);
+
 	GenerateTransformKeysForControl(TrackName, FrameNumbers, OutTransforms);
 }
 
 void UAnimationSequencerDataModel::GetBoneTrackTransforms(FName TrackName, TArray<FTransform>& OutTransforms) const
 {
+	UE::TUniqueLock<UE::FRecursiveMutex> Lock(EvaluationLock);
+
 	IterateTransformControlCurve(TrackName, [&OutTransforms](const FTransform& Transform, const FFrameNumber& FrameNumber) -> void
 	{
 		OutTransforms.Add(Transform);
@@ -537,6 +547,8 @@ void UAnimationSequencerDataModel::GetBoneTrackTransforms(FName TrackName, TArra
 
 void UAnimationSequencerDataModel::GetBoneTracksTransform(const TArray<FName>& TrackNames, const FFrameNumber& FrameNumber, TArray<FTransform>& OutTransforms) const
 {
+	UE::TUniqueLock<UE::FRecursiveMutex> Lock(EvaluationLock);
+
 	const TArray<FFrameNumber> FrameNumbers = { FrameNumber };
 	for (int32 EntryIndex = 0; EntryIndex < TrackNames.Num(); ++EntryIndex)
 	{
@@ -546,6 +558,7 @@ void UAnimationSequencerDataModel::GetBoneTracksTransform(const TArray<FName>& T
 
 FTransform UAnimationSequencerDataModel::EvaluateBoneTrackTransform(FName TrackName, const FFrameTime& FrameTime, const EAnimInterpolationType& Interpolation) const
 {
+	UE::TUniqueLock<UE::FRecursiveMutex> Lock(EvaluationLock);
 	const float Alpha = Interpolation == EAnimInterpolationType::Step ? FMath::RoundToFloat(FrameTime.GetSubFrame()) : FrameTime.GetSubFrame();
 
 	if (FMath::IsNearlyEqual(Alpha, 1.0f))
@@ -773,7 +786,7 @@ void UAnimationSequencerDataModel::Evaluate(FAnimationPoseData& InOutPoseData, c
 
 	if (UMovieSceneControlRigParameterTrack* Track = GetControlRigTrack())
 	{
-		FScopeLock Lock(&EvaluationLock);
+		UE::TUniqueLock<UE::FRecursiveMutex> Lock(EvaluationLock);
 
 		if(!bRigHierarchyInitialized)
 		{
@@ -1302,7 +1315,7 @@ void UAnimationSequencerDataModel::ClearControlRigData()
 {
 	if (bRigHierarchyInitialized)
 	{
-		FScopeLock Lock(&EvaluationLock);
+		UE::TUniqueLock<UE::FRecursiveMutex> Lock(EvaluationLock);
 		if (UFKControlRig* FKControlRig = Cast<UFKControlRig>(GetControlRig()))
 		{
 			FKControlRig->GetHierarchy()->Reset();
diff --git a/Engine/Plugins/Animation/AnimationData/Source/AnimationData/Private/AnimSequencerDataModel.h b/Engine/Plugins/Animation/AnimationData/Source/AnimationData/Private/AnimSequencerDataModel.h
--- a/Engine/Plugins/Animation/AnimationData/Source/AnimationData/Private/AnimSequencerDataModel.h
+++ b/Engine/Plugins/Animation/AnimationData/Source/AnimationData/Private/AnimSequencerDataModel.h
@@ -6,6 +6,8 @@
 #include "Animation/AnimData/IAnimationDataModel.h"
 #include "Sections/MovieSceneParameterSection.h"
 #include "Animation/AnimData/AnimDataModelNotifyCollector.h"
+#include "Async/RecursiveMutex.h"
+#include "Async/UniqueLock.h"
 
 #include "AnimSequencerDataModel.generated.h"
 
@@ -216,7 +218,7 @@ private:
 	mutable bool bRigHierarchyInitialized = false;
 
 	// Scope lock to prevent contention around ControlRig->Evaluation()
-	mutable FCriticalSection EvaluationLock;
+	mutable UE::FRecursiveMutex EvaluationLock;
 
 	friend class UAnimSequencerController;
 };
