fix up missing calls to initialize FK control rig hierarchy when running in lazy init mode

#Synced-CL 51330716

diff --git a/Engine/Plugins/Animation/AnimationData/Source/AnimationData/Private/AnimSequencerDataModel.cpp b/Engine/Plugins/Animation/AnimationData/Source/AnimationData/Private/AnimSequencerDataModel.cpp
--- a/Engine/Plugins/Animation/AnimationData/Source/AnimationData/Private/AnimSequencerDataModel.cpp
+++ b/Engine/Plugins/Animation/AnimationData/Source/AnimationData/Private/AnimSequencerDataModel.cpp
@@ -67,7 +67,7 @@ void UAnimationSequencerDataModel::RemoveOutOfDateControls() const
 	{
 		if (UFKControlRig* ControlRig = Cast<UFKControlRig>(Section->GetControlRig()))
 		{
-			if (URigHierarchy* Hierarchy = ControlRig->GetHierarchy())
+			if (URigHierarchy* Hierarchy = GetControlRigHierarchy())
 			{
 				if (URigHierarchyController* Controller = Hierarchy->GetController())
 				{
@@ -76,21 +76,21 @@ void UAnimationSequencerDataModel::RemoveOutOfDateControls() const
 					{
 						const bool bContainsBone = Section->HasTransformParameter(ControlElement->GetFName());
 						const bool bContainsCurve = Section->HasScalarParameter(ControlElement->GetFName());
-						
+
 						if (!bContainsBone && !bContainsCurve)
 						{
 							ElementKeysToRemove.Add(ControlElement->GetKey());
 						}
-						
+
 						return true;
 					});
-						
+
 					Hierarchy->ForEach<FRigCurveElement>([this, &ElementKeysToRemove](const FRigCurveElement* CurveElement) -> bool
 					{
 						const FName TargetCurveName = CurveElement->GetFName();
-						if(!LegacyCurveData.FloatCurves.ContainsByPredicate([TargetCurveName](const FFloatCurve& Curve) { return Curve.GetName() == TargetCurveName; }))
+						if (!LegacyCurveData.FloatCurves.ContainsByPredicate([TargetCurveName](const FFloatCurve& Curve) { return Curve.GetName() == TargetCurveName; }))
 						{
-							ElementKeysToRemove.Add(CurveElement->GetKey());	
+							ElementKeysToRemove.Add(CurveElement->GetKey());
 						}
 						return true;
 					});
@@ -822,6 +822,11 @@ void UAnimationSequencerDataModel::OnNotify(const EAnimDataModelNotifyType& Noti
 			    
 					if (UFKControlRig* FKRig = Cast<UFKControlRig>(Section->GetControlRig()))
 					{
+						if (!bRigHierarchyInitialized)
+						{
+							InitializeRigHierarchy(FKRig, GetSkeleton());
+						}
+
 						FKRig->RefreshActiveControls();
 					}
 				}
@@ -916,17 +921,20 @@ URigHierarchy* UAnimationSequencerDataModel::GetControlRigHierarchy() const
 	if (const UMovieSceneControlRigParameterSection* Section = GetFKControlRigSection())
 	{
 		UControlRig* ControlRig = Section->GetControlRig();
-		if (UFKControlRig* FKRig = Cast<UFKControlRig>(ControlRig))
+		if (ControlRig)
 		{
-			if (!bRigHierarchyInitialized)
+			if (UFKControlRig* FKRig = Cast<UFKControlRig>(ControlRig))
 			{
-				InitializeRigHierarchy(FKRig, GetSkeleton());
+				if (!bRigHierarchyInitialized)
+				{
+					InitializeRigHierarchy(FKRig, GetSkeleton());
+				}
+
+				return FKRig->GetHierarchy();
 			}
 
-			return FKRig->GetHierarchy();
+			IAnimationDataController::ReportObjectErrorf(this, LOCTEXT("UnableToFindRigHierarchy", "Unable to retrieve RigHierarchy for ControlRig ({0})"), FText::FromString(ControlRig->GetPathName()));
 		}
-
-		IAnimationDataController::ReportObjectErrorf(this, LOCTEXT("UnableToFindRigHierarchy", "Unable to retrieve RigHierarchy for ControlRig ({0})"), FText::FromString(ControlRig->GetPathName()));
 	}
 
 	return nullptr;
@@ -943,46 +951,39 @@ void UAnimationSequencerDataModel::RegenerateLegacyCurveData()
 		{
 			if (const UMovieSceneControlRigParameterSection* Section = Cast<UMovieSceneControlRigParameterSection>(TrackSection))
 			{
-				if (const UControlRig* ControlRig = Section->GetControlRig())
+				if (URigHierarchy* Hierarchy = GetControlRigHierarchy())
 				{
-					if(URigHierarchy* Hierarchy = ControlRig->GetHierarchy())
-					{
-						const TArray<FScalarParameterNameAndCurve>& ScalarCurves = Section->GetScalarParameterNamesAndCurves();
-						LegacyCurveData.FloatCurves.Empty();
+					const TArray<FScalarParameterNameAndCurve>& ScalarCurves = Section->GetScalarParameterNamesAndCurves();
+					LegacyCurveData.FloatCurves.Empty();
 
-						Hierarchy->ForEach<FRigCurveElement>([Hierarchy, this, ScalarCurves, FrameRate = GetFrameRate()](const FRigCurveElement* CurveElement) -> bool
+					Hierarchy->ForEach<FRigCurveElement>([Hierarchy, this, ScalarCurves, FrameRate = GetFrameRate()](const FRigCurveElement* CurveElement) -> bool
+					{
+						const FRigElementKey ControlKey(UFKControlRig::GetControlName(CurveElement->GetFName(), ERigElementType::Curve), ERigElementType::Control);
+						if (const FRigControlElement* Element = Hierarchy->Find<FRigControlElement>(ControlKey))
 						{
-							const FRigElementKey ControlKey(UFKControlRig::GetControlName(CurveElement->GetFName(), ERigElementType::Curve), ERigElementType::Control);
-							if (const FRigControlElement* Element = Hierarchy->Find<FRigControlElement>(ControlKey))
-							{
-								FFloatCurve& FloatCurve = LegacyCurveData.FloatCurves.AddDefaulted_GetRef();
-								FloatCurve.SetName(CurveElement->GetFName());
-								FloatCurve.Color = Element->Settings.ShapeColor;
+							FFloatCurve& FloatCurve = LegacyCurveData.FloatCurves.AddDefaulted_GetRef();
+							FloatCurve.SetName(CurveElement->GetFName());
+							FloatCurve.Color = Element->Settings.ShapeColor;
 
-								const FAnimationCurveIdentifier CurveId(FloatCurve.GetName(), ERawCurveTrackTypes::RCT_Float);
-								if (CurveIdentifierToMetaData.Contains(CurveId))
-								{
-									const FAnimationCurveMetaData& CurveMetaData = CurveIdentifierToMetaData.FindChecked(CurveId);
-									FloatCurve.SetCurveTypeFlags(CurveMetaData.Flags);
-									FloatCurve.Color = CurveMetaData.Color;
-									FloatCurve.Comment = CurveMetaData.Comment;
-								}
+							const FAnimationCurveIdentifier CurveId(FloatCurve.GetName(), ERawCurveTrackTypes::RCT_Float);
+							if (CurveIdentifierToMetaData.Contains(CurveId))
+							{
+								const FAnimationCurveMetaData& CurveMetaData = CurveIdentifierToMetaData.FindChecked(CurveId);
+								FloatCurve.SetCurveTypeFlags(CurveMetaData.Flags);
+								FloatCurve.Color = CurveMetaData.Color;
+								FloatCurve.Comment = CurveMetaData.Comment;
+							}
 
-								if (const FScalarParameterNameAndCurve* ScalarCurve = ScalarCurves.FindByPredicate([Element](const FScalarParameterNameAndCurve& Curve)
+							if (const FScalarParameterNameAndCurve* ScalarCurve = ScalarCurves.FindByPredicate([Element](const FScalarParameterNameAndCurve& Curve)
 								{
 									return Curve.ParameterName == Element->GetFName();
 								}))
-								{
-									AnimSequencerHelpers::ConvertFloatChannelToRichCurve(ScalarCurve->ParameterCurve, FloatCurve.FloatCurve, FrameRate);
-								}
+							{
+								AnimSequencerHelpers::ConvertFloatChannelToRichCurve(ScalarCurve->ParameterCurve, FloatCurve.FloatCurve, FrameRate);
 							}
-							return true;
-						});	
-					}
-					else
-					{
-						IAnimationDataController::ReportObjectErrorf(this, LOCTEXT("UnableToFindRigHierarchy", "Unable to retrieve RigHierarchy for ControlRig ({0})"), FText::FromString(ControlRig->GetPathName()));	      
-					}
+						}
+						return true;
+					});
 				}
 			}
 		}
@@ -994,37 +995,21 @@ void UAnimationSequencerDataModel::UpdateLegacyCurveData()
 	QUICK_SCOPE_CYCLE_COUNTER(STAT_UpdateLegacyCurveData);
 	ValidateSequencerData();
 
-	if (const UMovieSceneControlRigParameterTrack* Track = GetControlRigTrack())
+	if (const URigHierarchy* Hierarchy = GetControlRigHierarchy())
 	{
-		for (const UMovieSceneSection* TrackSection : Track->GetAllSections())
+		for (FFloatCurve& FloatCurve : LegacyCurveData.FloatCurves)
 		{
-			if (const UMovieSceneControlRigParameterSection* Section = Cast<UMovieSceneControlRigParameterSection>(TrackSection))
+			const FRigElementKey ControlKey(UFKControlRig::GetControlName(FloatCurve.GetName(), ERigElementType::Curve), ERigElementType::Control);
+			if (const FRigControlElement* Element = Hierarchy->Find<FRigControlElement>(ControlKey))
 			{
-				if (const UControlRig* ControlRig = Section->GetControlRig())
-				{
-					if(URigHierarchy* Hierarchy = ControlRig->GetHierarchy())
-					{
-						for(FFloatCurve& FloatCurve : LegacyCurveData.FloatCurves)
-						{
-							const FRigElementKey ControlKey(UFKControlRig::GetControlName(FloatCurve.GetName(), ERigElementType::Curve), ERigElementType::Control);
-							if (const FRigControlElement* Element = Hierarchy->Find<FRigControlElement>(ControlKey))
-							{
-								FloatCurve.Color = Element->Settings.ShapeColor;
+				FloatCurve.Color = Element->Settings.ShapeColor;
 
-								const FAnimationCurveIdentifier CurveId(FloatCurve.GetName(), ERawCurveTrackTypes::RCT_Float);
-								if (const FAnimationCurveMetaData* CurveMetaData = CurveIdentifierToMetaData.Find(CurveId))
-								{
-									FloatCurve.SetCurveTypeFlags(CurveMetaData->Flags);
-									FloatCurve.Color = CurveMetaData->Color;
-									FloatCurve.Comment = CurveMetaData->Comment;
-								}
-							}
-						}
-					}
-					else
-					{
-						IAnimationDataController::ReportObjectErrorf(this, LOCTEXT("UnableToFindRigHierarchy", "Unable to retrieve RigHierarchy for ControlRig ({0})"), FText::FromString(ControlRig->GetPathName()));	      
-					}
+				const FAnimationCurveIdentifier CurveId(FloatCurve.GetName(), ERawCurveTrackTypes::RCT_Float);
+				if (const FAnimationCurveMetaData* CurveMetaData = CurveIdentifierToMetaData.Find(CurveId))
+				{
+					FloatCurve.SetCurveTypeFlags(CurveMetaData->Flags);
+					FloatCurve.Color = CurveMetaData->Color;
+					FloatCurve.Comment = CurveMetaData->Comment;
 				}
 			}
 		}
@@ -1161,65 +1146,68 @@ void UAnimationSequencerDataModel::ValidateControlRigData() const
 
 void UAnimationSequencerDataModel::ValidateLegacyAgainstControlRigData() const
 {
-	UMovieSceneControlRigParameterSection* Section = GetFKControlRigSection();
+	if (bRigHierarchyInitialized)
+	{
+		UMovieSceneControlRigParameterSection* Section = GetFKControlRigSection();
 
-	UControlRig* ControlRig = Section->GetControlRig();
-	URigHierarchy* Hierarchy = ControlRig->GetHierarchy();
+		UControlRig* ControlRig = Section->GetControlRig();
+		URigHierarchy* Hierarchy = ControlRig->GetHierarchy();
 
-	// Validate bone tracks against controls
-	const UAnimSequence* OuterSequence = GetAnimationSequence();
-	if (const USkeleton* Skeleton = OuterSequence->GetSkeleton())
-	{
-		const FReferenceSkeleton& ReferenceSkeleton = Skeleton->GetReferenceSkeleton();	
-		// Validate curve data against controls
-		for (const FFloatCurve& FloatCurve : LegacyCurveData.FloatCurves)
+		// Validate bone tracks against controls
+		const UAnimSequence* OuterSequence = GetAnimationSequence();
+		if (const USkeleton* Skeleton = OuterSequence->GetSkeleton())
 		{
-			const FName CurveName = FloatCurve.GetName();					
-			const FRigElementKey CurveKey(CurveName, ERigElementType::Curve);
-			const FRigCurveElement* CurveElement = Hierarchy->Find<FRigCurveElement>(CurveKey);
-			if (!CurveElement)
+			const FReferenceSkeleton& ReferenceSkeleton = Skeleton->GetReferenceSkeleton();	
+			// Validate curve data against controls
+			for (const FFloatCurve& FloatCurve : LegacyCurveData.FloatCurves)
 			{
-				IAnimationDataController::ReportObjectErrorf(this, LOCTEXT("CurveElementNotFound", "Unable to find FRigCurve in RigHierarchy for Curve with name: {0}"), FText::FromName(CurveName));
-			}
+				const FName CurveName = FloatCurve.GetName();					
+				const FRigElementKey CurveKey(CurveName, ERigElementType::Curve);
+				const FRigCurveElement* CurveElement = Hierarchy->Find<FRigCurveElement>(CurveKey);
+				if (!CurveElement)
+				{
+					IAnimationDataController::ReportObjectErrorf(this, LOCTEXT("CurveElementNotFound", "Unable to find FRigCurve in RigHierarchy for Curve with name: {0}"), FText::FromName(CurveName));
+				}
 	        
-	        const FRigElementKey CurveControlKey(UFKControlRig::GetControlName(CurveName, ERigElementType::Curve), ERigElementType::Control);
-	        const FRigControlElement* CurveControlElement = Hierarchy->Find<FRigControlElement>(CurveControlKey);
-			if (!CurveControlElement)
-			{
-				IAnimationDataController::ReportObjectErrorf(this, LOCTEXT("CurveControlElementNotFound", "Unable to find FRigControlElement in RigHierarchy for Curve with name: {0}"), FText::FromName(CurveName));
-			}
+				const FRigElementKey CurveControlKey(UFKControlRig::GetControlName(CurveName, ERigElementType::Curve), ERigElementType::Control);
+				const FRigControlElement* CurveControlElement = Hierarchy->Find<FRigControlElement>(CurveControlKey);
+				if (!CurveControlElement)
+				{
+					IAnimationDataController::ReportObjectErrorf(this, LOCTEXT("CurveControlElementNotFound", "Unable to find FRigControlElement in RigHierarchy for Curve with name: {0}"), FText::FromName(CurveName));
+				}
 			
-			const FScalarParameterNameAndCurve* CurveControlParameter = Section->GetScalarParameterNamesAndCurves().FindByPredicate([CurveControlKey](const FScalarParameterNameAndCurve& ParameterPair)
-			{
-				return ParameterPair.ParameterName == CurveControlKey.Name;
-			});
+				const FScalarParameterNameAndCurve* CurveControlParameter = Section->GetScalarParameterNamesAndCurves().FindByPredicate([CurveControlKey](const FScalarParameterNameAndCurve& ParameterPair)
+				{
+					return ParameterPair.ParameterName == CurveControlKey.Name;
+				});
 			
-			if (CurveControlParameter)
-			{
-				for (const FRichCurveKey& Key : FloatCurve.FloatCurve.GetConstRefOfKeys())
+				if (CurveControlParameter)
 				{
-					float ParameterValue = 0.f;
-					const FFrameTime FrameTime = CurveControlParameter->ParameterCurve.GetTickResolution().AsFrameTime(Key.Time);
-
-					if (!CurveControlParameter->ParameterCurve.Evaluate(FrameTime, ParameterValue))
+					for (const FRichCurveKey& Key : FloatCurve.FloatCurve.GetConstRefOfKeys())
 					{
-						IAnimationDataController::ReportObjectErrorf(this, LOCTEXT("FailedToEvaluateCurveControl", "Unable to evaluate Control Curve ({0}) at interval {1}"), FText::FromName(CurveName), FText::AsNumber(FrameTime.AsDecimal()));
-					}
+						float ParameterValue = 0.f;
+						const FFrameTime FrameTime = CurveControlParameter->ParameterCurve.GetTickResolution().AsFrameTime(Key.Time);
 
-					const float RichCurveValue = FloatCurve.FloatCurve.Eval(Key.Time);
-					// QQ threshold
-					if (!(FMath::IsNearlyEqual(ParameterValue, Key.Value, 0.001f) || FMath::IsNearlyEqual(ParameterValue, RichCurveValue, 0.001f)))
-					{
-						IAnimationDataController::ReportObjectErrorf(this, LOCTEXT("CurveDeviationError", "Unexpected Control Curve ({0}) evaluation value {1} at {2}, expected {3} ({4})"), FText::FromName(CurveName), FText::AsNumber(ParameterValue), FText::AsNumber(FrameTime.AsDecimal()), FText::AsNumber(Key.Value), FText::AsNumber(RichCurveValue));
+						if (!CurveControlParameter->ParameterCurve.Evaluate(FrameTime, ParameterValue))
+						{
+							IAnimationDataController::ReportObjectErrorf(this, LOCTEXT("FailedToEvaluateCurveControl", "Unable to evaluate Control Curve ({0}) at interval {1}"), FText::FromName(CurveName), FText::AsNumber(FrameTime.AsDecimal()));
+						}
+
+						const float RichCurveValue = FloatCurve.FloatCurve.Eval(Key.Time);
+						// QQ threshold
+						if (!(FMath::IsNearlyEqual(ParameterValue, Key.Value, 0.001f) || FMath::IsNearlyEqual(ParameterValue, RichCurveValue, 0.001f)))
+						{
+							IAnimationDataController::ReportObjectErrorf(this, LOCTEXT("CurveDeviationError", "Unexpected Control Curve ({0}) evaluation value {1} at {2}, expected {3} ({4})"), FText::FromName(CurveName), FText::AsNumber(ParameterValue), FText::AsNumber(FrameTime.AsDecimal()), FText::AsNumber(Key.Value), FText::AsNumber(RichCurveValue));
+						}
 					}
 				}
-			}
-			else
-			{
-				IAnimationDataController::ReportObjectErrorf(this, LOCTEXT("ParameterNotFound", "Unable to find FScalarParameterNameAndCurve in RigHierarchy for Curve Control with name: {0}"), FText::FromName(CurveName));
-			}
-		}	
-	}	
+				else
+				{
+					IAnimationDataController::ReportObjectErrorf(this, LOCTEXT("ParameterNotFound", "Unable to find FScalarParameterNameAndCurve in RigHierarchy for Curve Control with name: {0}"), FText::FromName(CurveName));
+				}
+			}	
+		}
+	}
 }
 
 void UAnimationSequencerDataModel::IterateTransformControlCurve(const FName& BoneName, TFunction<void(const FTransform&, const FFrameNumber&)> IterationFunction, const TArray<FFrameNumber>* InFrameNumbers /*= nullptr*/) const
@@ -1778,119 +1766,119 @@ void UAnimationSequencerDataModel::IterateBoneKeys(const FName& BoneName, TFunct
 
 	if (const UMovieSceneControlRigParameterSection* Section = GetFKControlRigSection())
 	{
-	    UControlRig* ControlRig = Section->GetControlRig();
-	    const URigHierarchy* Hierarchy = ControlRig->GetHierarchy();
-	    
-	    const FRigElementKey BoneControlKey(UFKControlRig::GetControlName(BoneName, ERigElementType::Bone), ERigElementType::Control);
-	    if (const FTransformParameterNameAndCurves* ControlCurvePtr = Section->GetTransformParameterNamesAndCurves().FindByPredicate([CurveName = BoneControlKey.Name](const FTransformParameterNameAndCurves& TransformParameter)
-	    {
-		    return TransformParameter.ParameterName == CurveName;
-	    }))
-	    {
-		    const FTransformParameterNameAndCurves& ControlCurve = *ControlCurvePtr;
-
-		    struct FChannelInfo
-		    {
-			    bool bConstant = true;
-			    bool bUniform = true;
-		    };			
-		    FChannelInfo PosChannels[3];
-		    FChannelInfo RotChannels[3];
-		    FChannelInfo ScaleChannels[3];
-
-		    const int32 NumberOfKeys = GetNumberOfKeys();
-
-		    int32 MaxNumberOfKeys = 1;
-		    TSet<FFrameNumber> FrameNumbers;
-		    for (int32 ChannelIndex = 0; ChannelIndex < 3; ++ChannelIndex)
-		    {
-			    auto ProcessChannel = [&MaxNumberOfKeys, NumberOfKeys](const FMovieSceneFloatChannel& Channel, FChannelInfo& Info)
-			    {
-				    if (Channel.HasAnyData())
-				    {
-					    if (Channel.GetNumKeys() == 0)
-					    {
-						    Info.bUniform = false;
-					    }
-					    else
-					    {
-						    Info.bConstant = false;
-						    if (Channel.GetNumKeys() != NumberOfKeys)
-						    {
-							    Info.bUniform = false;
-						    }
-
-						    MaxNumberOfKeys = FMath::Max(MaxNumberOfKeys, Channel.GetNumKeys());
-					    }
-				    }
-			    };
-
-			    ProcessChannel(ControlCurve.Translation[ChannelIndex], PosChannels[ChannelIndex]);
-			    ProcessChannel(ControlCurve.Rotation[ChannelIndex], RotChannels[ChannelIndex]);
-			    ProcessChannel(ControlCurve.Scale[ChannelIndex], ScaleChannels[ChannelIndex]);
-		    }
-
-
-		    const int32 NumberOfKeysToIterate = MaxNumberOfKeys;
-		    const FRigControlElement* BoneControl = Hierarchy->Find<FRigControlElement>(BoneControlKey);
-		    const FTransform InitialTransform = BoneControl ? BoneControl->GetTransform().Initial.Local.Get() : FTransform::Identity;
-
-		    // Initialize components with initial control values (in case there is no default value nor keys)
-		    FVector3f PreviousPos = FVector3f(InitialTransform.GetLocation());
-		    FVector3f PreviousScale = FVector3f(InitialTransform.GetScale3D());
-		    FVector3f PreviousRot = FVector3f(InitialTransform.GetRotation().Euler());
-		    FQuat4f PreviousQuat;
-
-		    for (int32 KeyIndex = 0; KeyIndex < NumberOfKeysToIterate; ++KeyIndex)
-		    {
-			    const FFrameNumber Frame = KeyIndex;
-			    
-			    for (int32 ChannelIndex = 0; ChannelIndex < 3; ++ChannelIndex)
-			    {
-				    if (PosChannels[ChannelIndex].bConstant)
-				    {
-					    if (KeyIndex == 0)
-					    {
-						    PreviousPos[ChannelIndex] = ControlCurve.Translation[ChannelIndex].GetDefault().Get(PreviousPos[ChannelIndex]);
-					    }
-				    }
-				    else
-				    {
-					    ControlCurve.Translation[ChannelIndex].Evaluate(Frame, PreviousPos[ChannelIndex]);
-				    }
-				    
-				    if (RotChannels[ChannelIndex].bConstant)
-				    {
-					    if (KeyIndex == 0)
-					    {
-						    PreviousRot[ChannelIndex] = ControlCurve.Rotation[ChannelIndex].GetDefault().Get(PreviousRot[ChannelIndex]);
-					    }
-				    }
-				    else
-				    {
-					    ControlCurve.Rotation[ChannelIndex].Evaluate(Frame, PreviousRot[ChannelIndex]);
-				    }
-
-				    if (ScaleChannels[ChannelIndex].bConstant)
-				    {
-					    if (KeyIndex == 0)
-					    {
-						    PreviousScale[ChannelIndex] = ControlCurve.Scale[ChannelIndex].GetDefault().Get(PreviousScale[ChannelIndex]);
-					    }
-				    }
-				    else
-				    {
-					    ControlCurve.Scale[ChannelIndex].Evaluate(Frame, PreviousScale[ChannelIndex]);
-				    }
-			    }
-
-			    PreviousQuat = FQuat4f::MakeFromEuler(PreviousRot);
-			    if(!IterationFunction(PreviousPos, PreviousQuat, PreviousScale, Frame))
-			    {
-				    return;
-			    }
-		    }
-	    }
+		if (const URigHierarchy* Hierarchy = GetControlRigHierarchy())
+		{
+			const FRigElementKey BoneControlKey(UFKControlRig::GetControlName(BoneName, ERigElementType::Bone), ERigElementType::Control);
+			if (const FTransformParameterNameAndCurves* ControlCurvePtr = Section->GetTransformParameterNamesAndCurves().FindByPredicate([CurveName = BoneControlKey.Name](const FTransformParameterNameAndCurves& TransformParameter)
+				{
+					return TransformParameter.ParameterName == CurveName;
+				}))
+			{
+				const FTransformParameterNameAndCurves& ControlCurve = *ControlCurvePtr;
+
+				struct FChannelInfo
+				{
+					bool bConstant = true;
+					bool bUniform = true;
+				};
+				FChannelInfo PosChannels[3];
+				FChannelInfo RotChannels[3];
+				FChannelInfo ScaleChannels[3];
+
+				const int32 NumberOfKeys = GetNumberOfKeys();
+
+				int32 MaxNumberOfKeys = 1;
+				TSet<FFrameNumber> FrameNumbers;
+				for (int32 ChannelIndex = 0; ChannelIndex < 3; ++ChannelIndex)
+				{
+					auto ProcessChannel = [&MaxNumberOfKeys, NumberOfKeys](const FMovieSceneFloatChannel& Channel, FChannelInfo& Info)
+						{
+							if (Channel.HasAnyData())
+							{
+								if (Channel.GetNumKeys() == 0)
+								{
+									Info.bUniform = false;
+								}
+								else
+								{
+									Info.bConstant = false;
+									if (Channel.GetNumKeys() != NumberOfKeys)
+									{
+										Info.bUniform = false;
+									}
+
+									MaxNumberOfKeys = FMath::Max(MaxNumberOfKeys, Channel.GetNumKeys());
+								}
+							}
+						};
+
+					ProcessChannel(ControlCurve.Translation[ChannelIndex], PosChannels[ChannelIndex]);
+					ProcessChannel(ControlCurve.Rotation[ChannelIndex], RotChannels[ChannelIndex]);
+					ProcessChannel(ControlCurve.Scale[ChannelIndex], ScaleChannels[ChannelIndex]);
+				}
+
+
+				const int32 NumberOfKeysToIterate = MaxNumberOfKeys;
+				const FRigControlElement* BoneControl = Hierarchy->Find<FRigControlElement>(BoneControlKey);
+				const FTransform InitialTransform = BoneControl ? BoneControl->GetTransform().Initial.Local.Get() : FTransform::Identity;
+
+				// Initialize components with initial control values (in case there is no default value nor keys)
+				FVector3f PreviousPos = FVector3f(InitialTransform.GetLocation());
+				FVector3f PreviousScale = FVector3f(InitialTransform.GetScale3D());
+				FVector3f PreviousRot = FVector3f(InitialTransform.GetRotation().Euler());
+				FQuat4f PreviousQuat;
+
+				for (int32 KeyIndex = 0; KeyIndex < NumberOfKeysToIterate; ++KeyIndex)
+				{
+					const FFrameNumber Frame = KeyIndex;
+
+					for (int32 ChannelIndex = 0; ChannelIndex < 3; ++ChannelIndex)
+					{
+						if (PosChannels[ChannelIndex].bConstant)
+						{
+							if (KeyIndex == 0)
+							{
+								PreviousPos[ChannelIndex] = ControlCurve.Translation[ChannelIndex].GetDefault().Get(PreviousPos[ChannelIndex]);
+							}
+						}
+						else
+						{
+							ControlCurve.Translation[ChannelIndex].Evaluate(Frame, PreviousPos[ChannelIndex]);
+						}
+
+						if (RotChannels[ChannelIndex].bConstant)
+						{
+							if (KeyIndex == 0)
+							{
+								PreviousRot[ChannelIndex] = ControlCurve.Rotation[ChannelIndex].GetDefault().Get(PreviousRot[ChannelIndex]);
+							}
+						}
+						else
+						{
+							ControlCurve.Rotation[ChannelIndex].Evaluate(Frame, PreviousRot[ChannelIndex]);
+						}
+
+						if (ScaleChannels[ChannelIndex].bConstant)
+						{
+							if (KeyIndex == 0)
+							{
+								PreviousScale[ChannelIndex] = ControlCurve.Scale[ChannelIndex].GetDefault().Get(PreviousScale[ChannelIndex]);
+							}
+						}
+						else
+						{
+							ControlCurve.Scale[ChannelIndex].Evaluate(Frame, PreviousScale[ChannelIndex]);
+						}
+					}
+
+					PreviousQuat = FQuat4f::MakeFromEuler(PreviousRot);
+					if (!IterationFunction(PreviousPos, PreviousQuat, PreviousScale, Frame))
+					{
+						return;
+					}
+				}
+			}
+		}
 	}
 }
 
