fix missing compatible skeleton conversion when retrieving attributes
#Synced-CL 51777488

diff --git a/Engine/Plugins/Animation/AnimationData/Source/AnimationData/Private/AnimSequencerDataModel.cpp b/Engine/Plugins/Animation/AnimationData/Source/AnimationData/Private/AnimSequencerDataModel.cpp
--- a/Engine/Plugins/Animation/AnimationData/Source/AnimationData/Private/AnimSequencerDataModel.cpp
+++ b/Engine/Plugins/Animation/AnimationData/Source/AnimationData/Private/AnimSequencerDataModel.cpp
@@ -27,6 +27,8 @@
 #include "Runtime/MovieScene/Private/Channels/MovieSceneCurveChannelImpl.h"
 #include "Containers/StackTracker.h"
 #include "AnimationCompression.h"
+#include "Animation/SkeletonRemapping.h"
+#include "Animation/SkeletonRemappingRegistry.h"
 
 #include UE_INLINE_GENERATED_CPP_BY_NAME(AnimSequencerDataModel)
 
@@ -1463,10 +1465,14 @@ void UAnimationSequencerDataModel::GeneratePoseData(UControlRig* ControlRig, FAn
 
 			{				
 				QUICK_SCOPE_CYCLE_COUNTER(STAT_SetAttributes);
+				const FSkeletonRemapping& SkeletonRemapping = UE::Anim::FSkeletonRemappingRegistry::Get().GetRemapping(GetSkeleton(), RequiredBones.GetSkeletonAsset());
+
 				// Evaluate attributes at requested time interval
 				for (const FAnimatedBoneAttribute& Attribute : AnimatedBoneAttributes)
 				{
-					const FCompactPoseBoneIndex PoseBoneIndex = RequiredBones.GetCompactPoseIndexFromSkeletonPoseIndex(FSkeletonPoseBoneIndex(Attribute.Identifier.GetBoneIndex()));
+					const int32 AttributeBoneIndex = Attribute.Identifier.GetBoneIndex();
+					const int32 RemappedSkeletonBoneIndex = SkeletonRemapping.IsValid() ? SkeletonRemapping.GetTargetSkeletonBoneIndex(AttributeBoneIndex) : AttributeBoneIndex;
+					const FCompactPoseBoneIndex PoseBoneIndex = RequiredBones.GetCompactPoseIndexFromSkeletonPoseIndex(FSkeletonPoseBoneIndex(RemappedSkeletonBoneIndex));
 					// Only add attribute if the bone its tied to exists in the currently evaluated set of bones
 					if(PoseBoneIndex.IsValid())
 					{
